Fixed multimaterials when importing through Datasmith. Make sure to set ImportedMaterialSlotName because its required when converting to MeshDescription.
Also restore MaterialIndexToImportIndex in the FRawMesh just in case we need it later on.

diff --git a/Engine/Plugins/Enterprise/DatasmithContent/Source/DatasmithContent/Private/ObjectTemplates/DatasmithStaticMeshTemplate.cpp b/Engine/Plugins/Enterprise/DatasmithContent/Source/DatasmithContent/Private/ObjectTemplates/DatasmithStaticMeshTemplate.cpp
--- a/Engine/Plugins/Enterprise/DatasmithContent/Source/DatasmithContent/Private/ObjectTemplates/DatasmithStaticMeshTemplate.cpp
+++ b/Engine/Plugins/Enterprise/DatasmithContent/Source/DatasmithContent/Private/ObjectTemplates/DatasmithStaticMeshTemplate.cpp
@@ -81,6 +81,8 @@ void FDatasmithStaticMaterialTemplate::Apply( FStaticMaterial* Destination, FDat
 {
 	DATASMITHOBJECTTEMPLATE_CONDITIONALSET( MaterialSlotName, Destination, PreviousTemplate );
 	DATASMITHOBJECTTEMPLATE_CONDITIONALSET( MaterialInterface, Destination, PreviousTemplate );
+
+	Destination->ImportedMaterialSlotName = MaterialSlotName; // Not editable by the user, so always set it
 }
 
 void FDatasmithStaticMaterialTemplate::Load( const FStaticMaterial& Source )
