Set the Datasmith import options that aren't displayed to the user as Transient so that they don't get saved in the config file.

diff --git a/Engine/Plugins/Enterprise/DatasmithContent/Source/DatasmithContent/Public/DatasmithImportOptions.h b/Engine/Plugins/Enterprise/DatasmithContent/Source/DatasmithContent/Public/DatasmithImportOptions.h
--- a/Engine/Plugins/Enterprise/DatasmithContent/Source/DatasmithContent/Public/DatasmithImportOptions.h
+++ b/Engine/Plugins/Enterprise/DatasmithContent/Source/DatasmithContent/Public/DatasmithImportOptions.h
@@ -126,10 +126,10 @@ struct DATASMITHCONTENT_API FDatasmithStaticMeshImportOptions
 	UPROPERTY(EditAnywhere, BlueprintReadWrite, Category = Lightmap)
 	EDatasmithImportLightmapMax MaxLightmapResolution;
 
-	UPROPERTY(BlueprintReadWrite, Category = Lightmap)
+	UPROPERTY(BlueprintReadWrite, Category = Lightmap, Transient)
 	bool bGenerateLightmapUVs;
 
-	UPROPERTY(BlueprintReadWrite, Category = Mesh)
+	UPROPERTY(BlueprintReadWrite, Category = Mesh, Transient)
 	bool bRemoveDegenerates;
 
 public:
