MDL importer shader fix
- Moved the shader file in the plugin that contains the assets that are using it.
- Changed DatasmithContent to load on PostConfigInit so that PluginManager can setup the proper shader paths.

diff --git a/Engine/Plugins/Enterprise/DatasmithContent/DatasmithContent.uplugin b/Engine/Plugins/Enterprise/DatasmithContent/DatasmithContent.uplugin
--- a/Engine/Plugins/Enterprise/DatasmithContent/DatasmithContent.uplugin
+++ b/Engine/Plugins/Enterprise/DatasmithContent/DatasmithContent.uplugin
@@ -19,7 +19,7 @@
 		{
 			"Name" : "DatasmithContent",
 			"Type" : "Runtime",
-			"LoadingPhase" : "Default"
+			"LoadingPhase" : "PostConfigInit"
 		},
 		{
 			"Name" : "DatasmithContentEditor",
diff --git a/Engine/Plugins/Enterprise/DatasmithContent/Shaders/Private/FlakesAccess.usf b/Engine/Plugins/Enterprise/DatasmithContent/Shaders/Private/FlakesAccess.usf
new file mode 100644
--- /dev/null
+++ b/Engine/Plugins/Enterprise/DatasmithContent/Shaders/Private/FlakesAccess.usf
@@ -0,0 +1,131 @@
+
+/**
+Flakes texture access based on light angular direction, uses a pseduo(i.e. in atlas format) 2D texture array
+to access flakes slices and finally applies bilinear interpolation.
+
+@param ThetaAngles - normalized(by Pi/2) thetaF and thetaI angles
+@param FlakesTexArray - texture 2d array in an atlas format
+@param FlakesTexCount - number of textures in the array
+@param FlakesTexValues - x = texture size of an entry in the 2d array, y = ThetaSliceLUT texel size
+@param TexCoord - UV coordinates
+@param Tiling - UV tiling factor
+@param NumThetaF - number of ThetaF bins
+@param NumThetaI- number of thetaI bins
+@param MaxThetaI - max thetaI bin that is used
+@param ThetaSliceLUT - 1d texture with the index into flake texture array for a given thetaI bin
+@return Bilinear interpolated value of the flakes.
+*/
+
+#define PI2 1.5707963267948966
+#define USE_GRADIENTS 1
+
+struct TextureArraySample
+{ 
+    float4 Sample(float2 UV, float Index)
+    {
+      float SideCount = sqrt(FlakesTexCount); // per side texture count
+      float InvSideCount = 1.0 / SideCount; 
+      float TexPadding = 8.0; // pixels padding count
+      float TexSize = FlakesTexValues.x;
+      float PaddedTexSize = TexSize + TexPadding * 2.0;
+      float AtlasTexelSize = 1.0 / (SideCount * PaddedTexSize);
+
+      float2 TexOffset = Index;
+      TexOffset.x = fmod(TexOffset.x, SideCount);
+      TexOffset.y = floor(TexOffset.y * InvSideCount);
+      TexOffset *= InvSideCount;
+      TexOffset += TexPadding * AtlasTexelSize;
+
+      float2 SampleTexCoord = frac(UV * Tiling);
+      SampleTexCoord *= InvSideCount * (TexSize / PaddedTexSize);
+      SampleTexCoord += TexOffset;
+
+#if USE_GRADIENTS == 1
+      float2 DdxUV = DDX(SampleTexCoord);
+      float2 DdyUV = DDY(SampleTexCoord);
+      return Texture2DSampleGrad(FlakesTexArray, FlakesTexArraySampler, SampleTexCoord, DdxUV, DdyUV);
+#else
+      return Texture2DSample(FlakesTexArray, FlakesTexArraySampler, SampleTexCoord);
+#endif
+    }  
+}; 
+
+struct ThetaSliceLUTSample
+{ 
+    float Sample(float BinThetaI)
+    {
+      float TexelSize = FlakesTexValues.y;
+      float Padding = 4.0;
+      float Index = BinThetaI * Padding * TexelSize + TexelSize;
+
+      float Value = Texture2DSampleLevel(ThetaSliceLUT, ThetaSliceLUTSampler, float2(Index, 0.0), 0.0).x;
+      Value = floor(Value * 255.0 + 0.5);
+      return Value;
+    }
+};
+TextureArraySample TextureArray;
+ThetaSliceLUTSample ThetaFiLUT;
+
+//ThetaF sampling defines the angular sampling, i.e. angular flake lifetime
+float ThetaF = clamp(ThetaAngles.x, 0.0, 1.0 / PI2) * NumThetaF + 0.5;
+float ThetaI = clamp(ThetaAngles.y, 0.0, 1.0 / PI2) * NumThetaF + 0.5;
+
+//bilinear interpolation indices
+float LowThetaF = floor(ThetaF);
+float HighThetaF = LowThetaF + 1.0;
+float LowThetaI = floor(ThetaI);
+float HighThetaI = LowThetaI + 1.0;
+
+float2 UV = TexCoord;
+if (NumThetaI < NumThetaF)
+{
+    if ((LowThetaI % 2) == 1 || (HighThetaI % 2) == 1) 
+    { 
+      //tweak for extra randomness, in the case of almost planar geometries with repeating patches
+      UV.xy = UV.yx;
+    }
+
+    //map to the original sampling
+    float InvNumThetaF = 1.0 / NumThetaF;
+    LowThetaI = int(floor(float(LowThetaI) * NumThetaI * InvNumThetaF));
+    HighThetaI = int(floor(float(HighThetaI) * NumThetaI * InvNumThetaF));
+}
+
+half3 ColorBottomLeft = 0.0;
+half3 ColorTopLeft = 0.0;
+half3 ColorBottomRight = 0.0;
+half3 ColorTopRight  = 0.0;
+//access flakes texture and stay in the correct slices(no slip over)
+if (LowThetaI < MaxThetaI)
+{
+    float LowThetaIndex = ThetaFiLUT.Sample(LowThetaI) + LowThetaF;
+    if (LowThetaIndex < ThetaFiLUT.Sample(LowThetaI + 1.0))
+    {
+        ColorBottomLeft = TextureArray.Sample(UV, LowThetaIndex).xyz;
+    }
+    float HighThetaIndex = ThetaFiLUT.Sample(LowThetaI) + HighThetaF;
+    if (HighThetaIndex < ThetaFiLUT.Sample(LowThetaI + 1))
+    {
+        ColorBottomRight = TextureArray.Sample(UV, HighThetaIndex).xyz;
+    }
+}
+
+if (HighThetaI < MaxThetaI)
+{
+    float LowThetaIndex = ThetaFiLUT.Sample(HighThetaI) + LowThetaF;
+    if (LowThetaIndex < ThetaFiLUT.Sample(HighThetaI + 1))
+    {
+        ColorTopLeft = TextureArray.Sample(UV, LowThetaIndex).xyz;
+    }
+    float HighThetaIndex = ThetaFiLUT.Sample(HighThetaI) + HighThetaF;
+    if (HighThetaIndex < ThetaFiLUT.Sample(HighThetaI + 1))
+    {
+        ColorTopRight = TextureArray.Sample(UV, HighThetaIndex).xyz;
+    }
+}
+//bilinear interpolation
+float DiffThetaF = ThetaF - float(LowThetaF);
+float DiffThetaI = ThetaI - float(LowThetaI);
+half3 Bottom = (1.0 - DiffThetaF) * ColorBottomLeft + DiffThetaF * ColorBottomRight;
+half3 Top = (1.0 - DiffThetaF) * ColorTopLeft + DiffThetaF * ColorTopRight;    
+return (1.0 - DiffThetaI) * Bottom + DiffThetaI * Top;
