Cherry-pick 4361879: Maps //Plugin/DatasmithContent virtual shader directory to DatasmithContent's physical Shaders directory.

diff --git a/Engine/Plugins/Enterprise/DatasmithContent/Source/DatasmithContent/DatasmithContent.Build.cs b/Engine/Plugins/Enterprise/DatasmithContent/Source/DatasmithContent/DatasmithContent.Build.cs
--- a/Engine/Plugins/Enterprise/DatasmithContent/Source/DatasmithContent/DatasmithContent.Build.cs
+++ b/Engine/Plugins/Enterprise/DatasmithContent/Source/DatasmithContent/DatasmithContent.Build.cs
@@ -13,15 +13,17 @@ namespace UnrealBuildTool.Rules
 					"CinematicCamera",
 					"Core",
 					"CoreUObject",
-					"Engine"
-				}
+					"Engine",
+                    "ShaderCore",
+                }
 			);
 
 			PrivateDependencyModuleNames.AddRange(
 				new string[]
 				{
-					"Landscape"
-				}
+					"Landscape",
+                    "Projects",
+                }
 			);
 		}
 	}
diff --git a/Engine/Plugins/Enterprise/DatasmithContent/Source/DatasmithContent/Private/DatasmithContentModule.cpp b/Engine/Plugins/Enterprise/DatasmithContent/Source/DatasmithContent/Private/DatasmithContentModule.cpp
--- a/Engine/Plugins/Enterprise/DatasmithContent/Source/DatasmithContent/Private/DatasmithContentModule.cpp
+++ b/Engine/Plugins/Enterprise/DatasmithContent/Source/DatasmithContent/Private/DatasmithContentModule.cpp
@@ -6,6 +6,9 @@
 #include "HAL/PlatformProcess.h"
 #include "Modules/ModuleManager.h"
 #include "Misc/Paths.h"
+#include "Interfaces/IPluginManager.h"
+#include "Misc/Paths.h"
+#include "ShaderCore.h"
 
 /**
  * DatasmithContent module implementation (private)
@@ -18,6 +21,10 @@ public:
 		// Create temporary directory which will be used by UDatasmithStaticMeshCADImportData to store transient data
 		TempDir = FPaths::Combine(FPaths::ProjectIntermediateDir(), TEXT("DatasmithContentTemp"));
 		IFileManager::Get().MakeDirectory(*TempDir);
+
+		// Maps virtual shader source directory /Plugin/DatasmithContent to the plugin's actual Shaders directory.
+		FString PluginShaderDir = FPaths::Combine(IPluginManager::Get().FindPlugin(TEXT("DatasmithContent"))->GetBaseDir(), TEXT("Shaders"));
+		AddShaderSourceDirectoryMapping(TEXT("/Plugin/DatasmithContent"), PluginShaderDir);
 	}
 
 	virtual void ShutdownModule() override
