Fixed wrong materials assignment after deleting Sections of StaticMesh, undoing and re-importing through Datasmith.

diff --git a/Engine/Plugins/Enterprise/DatasmithContent/Source/DatasmithContent/DatasmithContent.Build.cs b/Engine/Plugins/Enterprise/DatasmithContent/Source/DatasmithContent/DatasmithContent.Build.cs
--- a/Engine/Plugins/Enterprise/DatasmithContent/Source/DatasmithContent/DatasmithContent.Build.cs
+++ b/Engine/Plugins/Enterprise/DatasmithContent/Source/DatasmithContent/DatasmithContent.Build.cs
@@ -23,6 +23,7 @@ namespace UnrealBuildTool.Rules
 				{
 					"Landscape",
 					"LevelSequence",
+					"MeshDescription",
                     "Projects",
                 }
 			);
diff --git a/Engine/Plugins/Enterprise/DatasmithContent/Source/DatasmithContent/Private/ObjectTemplates/DatasmithStaticMeshTemplate.cpp b/Engine/Plugins/Enterprise/DatasmithContent/Source/DatasmithContent/Private/ObjectTemplates/DatasmithStaticMeshTemplate.cpp
--- a/Engine/Plugins/Enterprise/DatasmithContent/Source/DatasmithContent/Private/ObjectTemplates/DatasmithStaticMeshTemplate.cpp
+++ b/Engine/Plugins/Enterprise/DatasmithContent/Source/DatasmithContent/Private/ObjectTemplates/DatasmithStaticMeshTemplate.cpp
@@ -6,6 +6,7 @@
 
 #include "CoreTypes.h"
 #include "Engine/StaticMesh.h"
+#include "MeshDescription.h"
 #include "Templates/Casts.h"
 
 FDatasmithMeshBuildSettingsTemplate::FDatasmithMeshBuildSettingsTemplate()
@@ -257,6 +258,42 @@ void UDatasmithStaticMeshTemplate::Apply( UObject* Destination, bool bForce )
 		}
 	}
 
+	// Make sure that the StaticMaterials are in the same order than the StaticMeshLODResources.Sections will be after the StaticMesh is built (see BuildVertexBuffer in StaticMeshBuilder.cpp)
+	const int32 NumLODs = StaticMesh->GetNumLODs();
+	for ( int32 LODIndex = 0; LODIndex < NumLODs; ++LODIndex )
+	{
+		FMeshDescription* MeshDescription = StaticMesh->GetMeshDescription( LODIndex );
+		if ( MeshDescription && MeshDescription->PolygonGroups().Num() == StaticMesh->StaticMaterials.Num() )
+		{
+			TArray< FStaticMaterial > TempStaticMaterials;
+			TPolygonGroupAttributesConstRef< FName > PolygonGroupImportedMaterialSlotNames = MeshDescription->PolygonGroupAttributes().GetAttributesRef< FName >( MeshAttribute::PolygonGroup::ImportedMaterialSlotName );
+			int32 SectionIndex = 0;
+			for ( const FPolygonGroupID PolygonGroupID : MeshDescription->PolygonGroups().GetElementIDs() )
+			{
+				int32 MaterialIndex = StaticMesh->GetMaterialIndexFromImportedMaterialSlotName( PolygonGroupImportedMaterialSlotNames[ PolygonGroupID ] );
+				if ( MaterialIndex == INDEX_NONE )
+				{
+					MaterialIndex = PolygonGroupID.GetValue();
+				}
+				TempStaticMaterials.Add( StaticMesh->StaticMaterials[ MaterialIndex ] );
+
+				// Note that the StaticMesh.SectionInfoMap MaterialIndex will overwrite the StaticMeshLODResources.Sections MaterialIndex through FStaticMeshRenderData::ResolveSectionInfo()
+				// This ensures there won't be any mismatch when that happens
+				FMeshSectionInfo SectionInfo = StaticMesh->SectionInfoMap.Get( LODIndex, SectionIndex );
+				SectionInfo.MaterialIndex = SectionIndex;
+				StaticMesh->SectionInfoMap.Set( LODIndex, SectionIndex, SectionInfo );
+
+				++SectionIndex;
+			}
+
+			// Set the StaticMaterials with respect to LOD 0
+			if ( LODIndex == 0 )
+			{
+				StaticMesh->StaticMaterials = MoveTemp( TempStaticMaterials );
+			}
+		}
+	}
+
 	FDatasmithObjectTemplateUtils::SetObjectTemplate( Destination, this );
 #endif // #if WITH_EDITORONLY_DATA
 }
