Hotfix: Fixed a crash when loading a new level after converting a datasmith scene actor to a blueprint.

The root cause was that the outer chain of datasmith user asset data from the template of the SCS_Node of the blueprint was keeping the old level alive.

diff --git a/Engine/Plugins/Enterprise/DatasmithContent/Source/DatasmithContent/Private/ObjectTemplates/DatasmithObjectTemplate.cpp b/Engine/Plugins/Enterprise/DatasmithContent/Source/DatasmithContent/Private/ObjectTemplates/DatasmithObjectTemplate.cpp
--- a/Engine/Plugins/Enterprise/DatasmithContent/Source/DatasmithContent/Private/ObjectTemplates/DatasmithObjectTemplate.cpp
+++ b/Engine/Plugins/Enterprise/DatasmithContent/Source/DatasmithContent/Private/ObjectTemplates/DatasmithObjectTemplate.cpp
@@ -64,6 +64,13 @@ TMap< TSubclassOf< UDatasmithObjectTemplate >, UDatasmithObjectTemplate* >* FDat
 	if (!UserData)
 	{
 		EObjectFlags Flags = RF_Public /*| RF_Transactional*/; // RF_Transactional Disabled as is can cause a crash in the transaction system for blueprints
+
+		if ( Outer->GetClass()->IsChildOf<AActor>() )
+		{
+			// The outer should never be an actor. (UE-70039)
+			Outer = static_cast<AActor*>(Outer)->GetRootComponent();
+		}
+
 		UserData = NewObject< UDatasmithAssetUserData >(Outer, NAME_None, Flags);
 		AssetUserDataInterface->AddAssetUserData(UserData);
 	}
