- Fixed IES textures not being set on the light components when going through Datasmith. The template generation didn't take the IESTexture property into account.
- Fixed dataprep assets not persisting. The RF_Standalone flag was missing.

diff --git a/Engine/Plugins/Enterprise/DatasmithContent/Source/DatasmithContent/Private/ObjectTemplates/DatasmithLightComponentTemplate.cpp b/Engine/Plugins/Enterprise/DatasmithContent/Source/DatasmithContent/Private/ObjectTemplates/DatasmithLightComponentTemplate.cpp
--- a/Engine/Plugins/Enterprise/DatasmithContent/Source/DatasmithContent/Private/ObjectTemplates/DatasmithLightComponentTemplate.cpp
+++ b/Engine/Plugins/Enterprise/DatasmithContent/Source/DatasmithContent/Private/ObjectTemplates/DatasmithLightComponentTemplate.cpp
@@ -75,6 +75,7 @@ void UDatasmithLightComponentTemplate::Load( const UObject* Source )
 	bUseTemperature = LightComponent->bUseTemperature;
 	Temperature = LightComponent->Temperature;
 
+	IESTexture = LightComponent->IESTexture;
 	bUseIESBrightness = LightComponent->bUseIESBrightness;
 	IESBrightnessScale = LightComponent->IESBrightnessScale;
 #endif // #if WITH_EDITORONLY_DATA
