Customization of the datasmith scene import data details for the dataprep editor.

diff --git a/Engine/Plugins/Enterprise/DatasmithContent/Source/DatasmithContentEditor/DatasmithContentEditor.Build.cs b/Engine/Plugins/Enterprise/DatasmithContent/Source/DatasmithContentEditor/DatasmithContentEditor.Build.cs
--- a/Engine/Plugins/Enterprise/DatasmithContent/Source/DatasmithContentEditor/DatasmithContentEditor.Build.cs
+++ b/Engine/Plugins/Enterprise/DatasmithContent/Source/DatasmithContentEditor/DatasmithContentEditor.Build.cs
@@ -13,6 +13,7 @@ namespace UnrealBuildTool.Rules
 					"Core",
 					"CoreUObject",
 					"DatasmithContent",
+					"DesktopPlatform",
 					"EditorStyle",
 					"Engine",
 					"Projects",
diff --git a/Engine/Plugins/Enterprise/DatasmithContent/Source/DatasmithContentEditor/Private/DatasmithSceneDetails.cpp b/Engine/Plugins/Enterprise/DatasmithContent/Source/DatasmithContentEditor/Private/DatasmithSceneDetails.cpp
new file mode 100644
--- /dev/null
+++ b/Engine/Plugins/Enterprise/DatasmithContent/Source/DatasmithContentEditor/Private/DatasmithSceneDetails.cpp
@@ -0,0 +1,248 @@
+// Copyright 1998-2019 Epic Games, Inc. All Rights Reserved.
+
+#include "DatasmithSceneDetails.h"
+
+#include "DatasmithAssetImportData.h"
+#include "DatasmithContentEditorModule.h"
+#include "DatasmithScene.h"
+
+#include "DesktopPlatformModule.h"
+#include "DetailCategoryBuilder.h"
+#include "DetailLayoutBuilder.h"
+#include "DetailWidgetRow.h"
+#include "EditorDirectories.h"
+#include "Framework/Application/SlateApplication.h"
+#include "IDesktopPlatform.h"
+#include "PropertyHandle.h"
+#include "Widgets/Images/SImage.h"
+#include "Widgets/Input/SButton.h"
+#include "Widgets/SBoxPanel.h"
+#include "Widgets/SWidget.h"
+#include "Widgets/Text/STextBlock.h"
+
+#define LOCTEXT_NAMESPACE "DatasmithSceneDetails"
+
+void FDatasmithSceneDetails::CustomizeDetails(IDetailLayoutBuilder& DetailBuilder)
+{
+	TArray< TWeakObjectPtr< UObject > > Objects;
+	DetailBuilder.GetObjectsBeingCustomized( Objects );
+	check( Objects.Num() > 0 );
+	UDatasmithScene* DatasmithScene = Cast< UDatasmithScene >( Objects[0].Get() );
+	check( DatasmithScene );
+	DatasmithScenePtr = DatasmithScene;
+
+	static_assert( TIsSame< TRemovePointer< decltype( UDatasmithScene::AssetImportData ) >::Type, UDatasmithSceneImportData >::Value, "Please update this details customization" );
+	static_assert( TIsDerivedFrom < UDatasmithSceneImportData, UAssetImportData >::IsDerived, "Please update this details customization" );
+
+	TSharedRef< IPropertyHandle > AssetImportDataPropertyHandle = DetailBuilder.GetProperty( GET_MEMBER_NAME_CHECKED( UDatasmithScene, AssetImportData ) );
+	AssetImportDataPropertyHandle->MarkHiddenByCustomization();
+
+	// Force the refresh of the customization when the User asset data have changed
+	AssetImportDataPropertyHandle->SetOnPropertyValueChanged( FSimpleDelegate::CreateSP( this, &FDatasmithSceneDetails::ForceRefreshDetails ) );
+
+	IDetailCategoryBuilder& ImportSettingsCategoryBuilder = DetailBuilder.EditCategory( AssetImportDataPropertyHandle->GetDefaultCategoryName(), FText::GetEmpty(), ECategoryPriority::Important );
+	FDetailWidgetRow& CustomAssetImportRow = ImportSettingsCategoryBuilder.AddCustomRow( FText::FromString( TEXT( "Import File" ) ) );
+
+	FString FilePath;
+	if ( DatasmithScene->AssetImportData && DatasmithScene->AssetImportData->SourceData.SourceFiles.Num() )
+	{
+		FilePath = DatasmithScene->AssetImportData->SourceData.SourceFiles[0].RelativeFilename;
+	}
+
+	CustomAssetImportRow.NameContent()
+	[
+		SNew(STextBlock)
+		.Text(LOCTEXT("ImportFile", "Import File"))
+		.ToolTipText(LOCTEXT("ImportFileTooltip", "The file imported by datasmith."))
+		.Font(DetailBuilder.GetDetailFont())
+	];
+
+	TSharedPtr< SWidget > TextBlock;
+	TSharedPtr< SButton > Button;
+	const float ButtonPaddingLeft = 4.f;
+
+	CustomAssetImportRow.ValueContent()
+	[
+		SNew(SHorizontalBox)
+		+ SHorizontalBox::Slot()
+		.FillWidth(1.f)
+		.VAlign(VAlign_Center)
+		.HAlign(HAlign_Left)
+		[
+			// file path text box
+			SAssignNew(TextBlock, STextBlock)
+			.Text(FilePath.IsEmpty() ? FText::FromString(TEXT("Select a file")) : FText::FromString(MoveTemp(FilePath)))
+			.Font(DetailBuilder.GetDetailFont())
+			.Justification(ETextJustify::Right)
+		]
+
+		+ SHorizontalBox::Slot()
+		.Padding(ButtonPaddingLeft, 0.0f, 0.0f, 0.0f)
+		.AutoWidth()
+		.VAlign(VAlign_Center)
+		.HAlign(HAlign_Left)
+		[
+			// browse button
+			SAssignNew(Button, SButton)
+			.ButtonStyle(FEditorStyle::Get(), "HoverHintOnly")
+			.ToolTipText(LOCTEXT("FileButtonToolTipText", "Choose a source import file"))
+			.OnClicked(this, &FDatasmithSceneDetails::OnSelectFile)
+			.ContentPadding(2.0f)
+			.ForegroundColor(FSlateColor::UseForeground())
+			.IsFocusable(false)
+			[
+				SNew(SImage)
+				.Image(FEditorStyle::GetBrush("PropertyWindow.Button_Ellipsis"))
+				.ColorAndOpacity(FSlateColor::UseForeground())
+			]
+		]
+	];
+
+	CustomAssetImportRow.ValueContent().MaxDesiredWidth( TextBlock->GetDesiredSize().X + Button->GetDesiredSize().X + ButtonPaddingLeft );
+
+	TSharedPtr< IPropertyHandle > DesiredParentHandle = AssetImportDataPropertyHandle->GetChildHandle( GET_MEMBER_NAME_CHECKED( UDatasmithSceneImportData, BaseOptions ) );
+	if ( DesiredParentHandle )
+	{
+		// Walk up until we have the right parent.
+		const FName NonDesiredPropertyName = GET_MEMBER_NAME_CHECKED( UDatasmithSceneImportData, SourceData );
+		while ( !DesiredParentHandle->GetChildHandle( NonDesiredPropertyName ) )
+		{
+			DesiredParentHandle = DesiredParentHandle->GetParentHandle();
+		}
+
+		uint32 NumChildren;
+		DesiredParentHandle->GetNumChildren( NumChildren );
+		for ( uint32 i = 1; i <= NumChildren; i++ )
+		{
+			// Going backward just for the display
+			TSharedPtr< IPropertyHandle > PropertyHandle = DesiredParentHandle->GetChildHandle( NumChildren - i );
+			if ( !PropertyHandle->GetChildHandle( NonDesiredPropertyName ) )
+			{
+				ImportSettingsCategoryBuilder.AddProperty( PropertyHandle )
+					.ShouldAutoExpand( true );
+			}
+		}
+	}
+}
+
+void FDatasmithSceneDetails::CustomizeDetails(const TSharedPtr<IDetailLayoutBuilder>& DetailBuilder)
+{
+	DetailBuilderWeakPtr = DetailBuilder;
+	CustomizeDetails( *DetailBuilder );
+}
+
+FReply FDatasmithSceneDetails::OnSelectFile()
+{
+	TArray< FImporterDescription > ImporterDescriptions = IDatasmithContentEditorModule::Get().GetDatasmithImporters();
+
+	FString FileTypes;
+	FString AllExtensions;
+
+	// Map a file extension to an importer
+	TMap< FString, TSharedPtr< IDataPrepImporterInterface > > ImportHandlerMap;
+
+	for ( const FImporterDescription& ImporterDescription : ImporterDescriptions )
+	{
+		if ( ImporterDescription.Handler.IsBound() )
+		{
+			TSharedPtr<IDataPrepImporterInterface> Importer = ImporterDescription.Handler.Execute();
+			if ( Importer.IsValid() )
+			{
+				for ( const FString& Format : ImporterDescription.Formats )
+				{
+					TArray<FString> FormatComponents;
+					Format.ParseIntoArray( FormatComponents, TEXT( ";" ), false );
+
+					for ( int32 ComponentIndex = 0; ComponentIndex < FormatComponents.Num(); ComponentIndex += 2 )
+					{
+						check( FormatComponents.IsValidIndex( ComponentIndex + 1 ) );
+						const FString& Extension = FormatComponents[ComponentIndex];
+						const FString& Description = FormatComponents[ComponentIndex + 1];
+
+						if ( !AllExtensions.IsEmpty() )
+						{
+							AllExtensions.AppendChar( TEXT( ';' ) );
+						}
+						AllExtensions.Append( TEXT( "*." ) );
+						AllExtensions.Append( Extension );
+
+						if ( !FileTypes.IsEmpty() )
+						{
+							FileTypes.AppendChar( TEXT( '|' ) );
+						}
+
+						FileTypes.Append( FString::Printf( TEXT( "%s (*.%s)|*.%s" ), *Description, *Extension, *Extension ) );
+
+						ImportHandlerMap.Add( Extension, Importer );
+					}
+				}
+			}
+		}
+	}
+
+	FString SupportedExtensions( FString::Printf( TEXT( "All Files (%s)|%s|%s" ), *AllExtensions, *AllExtensions, *FileTypes ) );
+
+	TArray<FString> OpenedFiles;
+	FString DefaultLocation( FEditorDirectories::Get().GetLastDirectory( ELastDirectory::GENERIC_IMPORT ) );
+	IDesktopPlatform* DesktopPlatform = FDesktopPlatformModule::Get();
+
+	bool bOpened = false;
+	if ( DesktopPlatform )
+	{
+		bOpened = DesktopPlatform->OpenFileDialog(
+			FSlateApplication::Get().FindBestParentWindowHandleForDialogs( nullptr ),
+			LOCTEXT( "FileDialogTitle", "Import Datasmith" ).ToString(),
+			DefaultLocation,
+			TEXT( "" ),
+			SupportedExtensions,
+			EFileDialogFlags::None,
+			OpenedFiles
+		);
+	}
+
+	if ( bOpened && OpenedFiles.Num() > 0 )
+	{
+		FString& OpenedFile = OpenedFiles[0];
+		FEditorDirectories::Get().SetLastDirectory( ELastDirectory::GENERIC_IMPORT, OpenedFile );
+
+		FString Extension = FPaths::GetExtension( OpenedFile );
+
+		if ( TSharedPtr< IDataPrepImporterInterface > ImportHandlerPtr = ImportHandlerMap.FindRef( Extension ) )
+		{
+			// Raw because we don't want to keep alive the details builder when calling the force refresh details
+			IDetailLayoutBuilder* DetailLayoutBuilder = DetailBuilderWeakPtr.Pin().Get();
+			if ( DetailLayoutBuilder )
+			{
+				if ( UDatasmithScene* DatasmithScene = DatasmithScenePtr.Get() )
+				{
+					UClass* OldAssetImportDataClass = DatasmithScene->AssetImportData ? DatasmithScene->AssetImportData->GetClass() : nullptr;
+					UClass* NewAssetImportDataClass = ImportHandlerPtr->GetAssetImportDataClass();
+
+					if ( OldAssetImportDataClass != NewAssetImportDataClass && NewAssetImportDataClass )
+					{
+						DatasmithScene->AssetImportData = NewObject< UDatasmithSceneImportData >( DatasmithScene, NewAssetImportDataClass );
+					}
+
+					check( DatasmithScene->AssetImportData );
+					DatasmithScene->AssetImportData->Update( OpenedFile );
+
+					ForceRefreshDetails();
+				}
+			}
+		}
+	}
+
+	return FReply::Handled();
+}
+
+void FDatasmithSceneDetails::ForceRefreshDetails()
+{
+	// Raw because we don't want to keep alive the details builder when calling the force refresh details
+	IDetailLayoutBuilder* DetailLayoutBuilder = DetailBuilderWeakPtr.Pin().Get();
+	if ( DetailLayoutBuilder )
+	{
+		DetailLayoutBuilder->ForceRefreshDetails();
+	}
+}
+
+#undef LOCTEXT_NAMESPACE
diff --git a/Engine/Plugins/Enterprise/DatasmithContent/Source/DatasmithContentEditor/Public/DatasmithContentEditorModule.h b/Engine/Plugins/Enterprise/DatasmithContent/Source/DatasmithContentEditor/Public/DatasmithContentEditorModule.h
--- a/Engine/Plugins/Enterprise/DatasmithContent/Source/DatasmithContentEditor/Public/DatasmithContentEditorModule.h
+++ b/Engine/Plugins/Enterprise/DatasmithContent/Source/DatasmithContentEditor/Public/DatasmithContentEditorModule.h
@@ -30,6 +30,7 @@ public:
 	virtual bool BuildWorld(TArray<TWeakObjectPtr<UObject>>& OutAssets) = 0;
 	virtual bool SetFinalWorld(UWorld* FinalWorld) = 0;
 	virtual bool FinalizeAssets(const TArray<TWeakObjectPtr<UObject>>& Assets) = 0;
+	virtual TSubclassOf<class UDatasmithSceneImportData> GetAssetImportDataClass() const = 0;
 };
 
 DECLARE_DELEGATE_TwoParams( FOnSpawnDatasmithSceneActors, class ADatasmithSceneActor*, bool );
diff --git a/Engine/Plugins/Enterprise/DatasmithContent/Source/DatasmithContentEditor/Public/DatasmithSceneDetails.h b/Engine/Plugins/Enterprise/DatasmithContent/Source/DatasmithContentEditor/Public/DatasmithSceneDetails.h
new file mode 100644
--- /dev/null
+++ b/Engine/Plugins/Enterprise/DatasmithContent/Source/DatasmithContentEditor/Public/DatasmithSceneDetails.h
@@ -0,0 +1,32 @@
+// Copyright 1998-2019 Epic Games, Inc. All Rights Reserved.
+
+#pragma once
+
+#include "CoreMinimal.h"
+#include "IDetailCustomization.h"
+#include "Input/Reply.h"
+
+class IDetailLayoutBuilder;
+
+// Customization of the details of the Datasmith Scene for the data prep editor.
+class DATASMITHCONTENTEDITOR_API FDatasmithSceneDetails : public IDetailCustomization
+{
+public:
+	static TSharedRef< IDetailCustomization > MakeDetails() { return MakeShared<FDatasmithSceneDetails>(); };
+
+	/** Called when details should be customized */
+	virtual void CustomizeDetails( IDetailLayoutBuilder& DetailBuilder ) override;
+
+	virtual void CustomizeDetails( const TSharedPtr<IDetailLayoutBuilder>& DetailBuilder ) override;
+
+private:
+
+	FReply OnSelectFile();
+
+	void ForceRefreshDetails();
+
+	TWeakPtr< IDetailLayoutBuilder > DetailBuilderWeakPtr;
+
+	TWeakObjectPtr< class UDatasmithScene > DatasmithScenePtr;
+};
+
