Fixed a issue with how we setted an SoftObjectPtr on a weak member in the datasmith templates.

diff --git a/Engine/Plugins/Enterprise/DatasmithContent/Source/DatasmithContent/Public/ObjectTemplates/DatasmithObjectTemplate.h b/Engine/Plugins/Enterprise/DatasmithContent/Source/DatasmithContent/Public/ObjectTemplates/DatasmithObjectTemplate.h
--- a/Engine/Plugins/Enterprise/DatasmithContent/Source/DatasmithContent/Public/ObjectTemplates/DatasmithObjectTemplate.h
+++ b/Engine/Plugins/Enterprise/DatasmithContent/Source/DatasmithContent/Public/ObjectTemplates/DatasmithObjectTemplate.h
@@ -58,7 +58,7 @@ public:
 #define DATASMITHOBJECTTEMPLATE_CONDITIONALSETSOFTOBJECTPTR(MemberName, Destination, PreviousTemplate) \
 	if ( !PreviousTemplate || Destination->MemberName == PreviousTemplate->MemberName.Get() ) \
 	{ \
-		Destination->MemberName = MemberName.Get(); \
+		Destination->MemberName = MemberName.LoadSynchronous(); \
 	}
 
 struct DATASMITHCONTENT_API FDatasmithObjectTemplateUtils
