Copy of 6911880 from //UE4/Dev-Enterprise:

Remove CADImporter module

diff --git a/Engine/Plugins/Enterprise/DatasmithContent/Source/DatasmithContent/Public/DatasmithImportOptions.h b/Engine/Plugins/Enterprise/DatasmithContent/Source/DatasmithContent/Public/DatasmithImportOptions.h
--- a/Engine/Plugins/Enterprise/DatasmithContent/Source/DatasmithContent/Public/DatasmithImportOptions.h
+++ b/Engine/Plugins/Enterprise/DatasmithContent/Source/DatasmithContent/Public/DatasmithImportOptions.h
@@ -329,9 +329,6 @@ public:
 	UPROPERTY(config, EditAnywhere, BlueprintReadWrite, Category = "Options", meta = (ShowOnlyInnerProperties))
 	FDatasmithImportBaseOptions BaseOptions;
 
-	UPROPERTY(config, EditAnywhere, BlueprintReadWrite, Category = "NotVisible", meta = (ShowOnlyInnerProperties))
-	FDatasmithTessellationOptions TessellationOptions;
-
 	/** Options specific to the reimport process */
 	UPROPERTY(config, EditAnywhere, BlueprintReadWrite, Category = "NotVisible", meta = (ShowOnlyInnerProperties))
 	FDatasmithReimportOptions ReimportOptions;
