Copy of 6941657 from //UE4/Dev-Enterprise-Minimal:

Fix Blueprint unable to edit datasmith import options

diff --git a/Engine/Plugins/Enterprise/DatasmithContent/Source/DatasmithContent/Public/DatasmithImportOptions.h b/Engine/Plugins/Enterprise/DatasmithContent/Source/DatasmithContent/Public/DatasmithImportOptions.h
--- a/Engine/Plugins/Enterprise/DatasmithContent/Source/DatasmithContent/Public/DatasmithImportOptions.h
+++ b/Engine/Plugins/Enterprise/DatasmithContent/Source/DatasmithContent/Public/DatasmithImportOptions.h
@@ -275,7 +275,7 @@ public:
 	}
 };
 
-UCLASS(config = Editor, Transient)
+UCLASS(BlueprintType, config = EditorPerProjectUserSettings, Transient)
 class DATASMITHCONTENT_API UDatasmithCommonTessellationOptions : public UObject
 {
 	GENERATED_BODY()
