Merging 6896081 from //UE4/Dev-Enterprise/... to //UE4/Release-4.23/...

Fixed most of the issues were some DatasmithAssetUserData would loose their templates.
The blueprint base objects or actors seem to still cause some issues.

diff --git a/Engine/Plugins/Enterprise/DatasmithContent/Source/DatasmithContent/Private/ObjectTemplates/DatasmithObjectTemplate.cpp b/Engine/Plugins/Enterprise/DatasmithContent/Source/DatasmithContent/Private/ObjectTemplates/DatasmithObjectTemplate.cpp
--- a/Engine/Plugins/Enterprise/DatasmithContent/Source/DatasmithContent/Private/ObjectTemplates/DatasmithObjectTemplate.cpp
+++ b/Engine/Plugins/Enterprise/DatasmithContent/Source/DatasmithContent/Private/ObjectTemplates/DatasmithObjectTemplate.cpp
@@ -8,35 +8,29 @@
 #include "GameFramework/Actor.h"
 
 
-namespace
+namespace DatasmithObjectTemplateUtilsImpl
 {
-
-IInterface_AssetUserData* GetUserDataInterface(UObject* Outer)
-{
-#if WITH_EDITORONLY_DATA
-	if (Outer && Outer->GetClass()->IsChildOf(AActor::StaticClass()))
+	IInterface_AssetUserData* GetUserDataInterface(UObject* Outer)
 	{
-		// The root Component holds AssetUserData on behalf of the actor
-		Outer = Cast<AActor>(Outer)->GetRootComponent();
-	}
-
-	if (!Outer || !Outer->GetClass()->ImplementsInterface(UInterface_AssetUserData::StaticClass()))
-	{
-		return nullptr;
-	}
+		if (Outer && Outer->GetClass()->IsChildOf(AActor::StaticClass()))
+		{
+			// The root Component holds AssetUserData on behalf of the actor
+			Outer = Cast<AActor>(Outer)->GetRootComponent();
+		}
 
-	return Cast< IInterface_AssetUserData >(Outer);
-#else
-	return nullptr;
-#endif // #if WITH_EDITORONLY_DATA
-}
+		if (!Outer || !Outer->GetClass()->ImplementsInterface(UInterface_AssetUserData::StaticClass()))
+		{
+			return nullptr;
+		}
 
+		return Cast< IInterface_AssetUserData >(Outer);
+	}
 }
 
 bool FDatasmithObjectTemplateUtils::HasObjectTemplates(UObject* Outer)
 {
 #if WITH_EDITORONLY_DATA
-	IInterface_AssetUserData* AssetUserDataInterface = GetUserDataInterface(Outer);
+	IInterface_AssetUserData* AssetUserDataInterface = DatasmithObjectTemplateUtilsImpl::GetUserDataInterface(Outer);
 	if (!AssetUserDataInterface)
 	{
 		return false;
@@ -50,11 +44,9 @@ bool FDatasmithObjectTemplateUtils::HasObjectTemplates(UObject* Outer)
 #endif // #if WITH_EDITORONLY_DATA
 }
 
-
-TMap< TSubclassOf< UDatasmithObjectTemplate >, UDatasmithObjectTemplate* >* FDatasmithObjectTemplateUtils::FindOrCreateObjectTemplates(UObject* Outer)
+UDatasmithAssetUserData* FDatasmithObjectTemplateUtils::FindOrCreateDatasmithUserData(UObject* Outer)
 {
-#if WITH_EDITORONLY_DATA
-	IInterface_AssetUserData* AssetUserDataInterface = GetUserDataInterface(Outer);
+	IInterface_AssetUserData* AssetUserDataInterface = DatasmithObjectTemplateUtilsImpl::GetUserDataInterface(Outer);
 	if (!AssetUserDataInterface)
 	{
 		return nullptr;
@@ -76,10 +68,18 @@ TMap< TSubclassOf< UDatasmithObjectTemplate >, UDatasmithObjectTemplate* >* FDat
 		AssetUserDataInterface->AddAssetUserData(UserData);
 	}
 
-	return &UserData->ObjectTemplates;
-#else
-	return nullptr;
+	return UserData;
+}
+
+TMap< TSubclassOf< UDatasmithObjectTemplate >, UDatasmithObjectTemplate* >* FDatasmithObjectTemplateUtils::FindOrCreateObjectTemplates(UObject* Outer)
+{
+#if WITH_EDITORONLY_DATA
+	if (UDatasmithAssetUserData* UserAssetData = FindOrCreateDatasmithUserData(Outer))
+	{
+		return &UserAssetData->ObjectTemplates;
+	}
 #endif // #if WITH_EDITORONLY_DATA
+return nullptr;
 }
 
 void UDatasmithObjectTemplate::Apply(UObject* Destination, bool bForce)
@@ -113,15 +113,18 @@ UDatasmithObjectTemplate* FDatasmithObjectTemplateUtils::GetObjectTemplate(UObje
 void FDatasmithObjectTemplateUtils::SetObjectTemplate(UObject* Outer, UDatasmithObjectTemplate* ObjectTemplate)
 {
 #if WITH_EDITORONLY_DATA
-	TMap< TSubclassOf< UDatasmithObjectTemplate >, UDatasmithObjectTemplate* >* ObjectTemplatesMap = FindOrCreateObjectTemplates(Outer);
-	ensure(ObjectTemplatesMap);
-
-	if (!ObjectTemplatesMap)
+	if (UDatasmithAssetUserData* UserData = FindOrCreateDatasmithUserData(Outer))
 	{
-		return;
-	}
+		TMap< TSubclassOf< UDatasmithObjectTemplate >, UDatasmithObjectTemplate* >& ObjectTemplatesMap = UserData->ObjectTemplates;
+		ObjectTemplatesMap.FindOrAdd(ObjectTemplate->GetClass()) = ObjectTemplate;
 
-	ObjectTemplatesMap->FindOrAdd(ObjectTemplate->GetClass()) = ObjectTemplate;
+		if (ObjectTemplate->GetOuter() != UserData)
+		{
+			// The outer chain is important for most of the engine functionality.
+			// If it's not set properly the deep copy of object won't work properly
+			ObjectTemplate->Rename(nullptr, UserData);
+		}
+	}
 #endif // #if WITH_EDITORONLY_DATA
 }
 
diff --git a/Engine/Plugins/Enterprise/DatasmithContent/Source/DatasmithContent/Public/ObjectTemplates/DatasmithObjectTemplate.h b/Engine/Plugins/Enterprise/DatasmithContent/Source/DatasmithContent/Public/ObjectTemplates/DatasmithObjectTemplate.h
--- a/Engine/Plugins/Enterprise/DatasmithContent/Source/DatasmithContent/Public/ObjectTemplates/DatasmithObjectTemplate.h
+++ b/Engine/Plugins/Enterprise/DatasmithContent/Source/DatasmithContent/Public/ObjectTemplates/DatasmithObjectTemplate.h
@@ -84,6 +84,8 @@ struct DATASMITHCONTENT_API FDatasmithObjectTemplateUtils
 {
 	static bool HasObjectTemplates( UObject* Outer );
 
+	static class UDatasmithAssetUserData* FindOrCreateDatasmithUserData( UObject* Outer );
+
 	static TMap< TSubclassOf< UDatasmithObjectTemplate >, UDatasmithObjectTemplate* >* FindOrCreateObjectTemplates( UObject* Outer );
 
 	static UDatasmithObjectTemplate* GetObjectTemplate( UObject* Outer, TSubclassOf< UDatasmithObjectTemplate > Subclass );
