Merging 6917731 from //UE4/Dev-Enterprise/... to //UE4/Release-4.23/...

Fixed an issue where the converted components with a datasmith user asset to a blueprint would block the level from saving (Only for those created by the datasmith importer).

diff --git a/Engine/Plugins/Enterprise/DatasmithContent/Source/DatasmithContent/Private/ObjectTemplates/DatasmithObjectTemplate.cpp b/Engine/Plugins/Enterprise/DatasmithContent/Source/DatasmithContent/Private/ObjectTemplates/DatasmithObjectTemplate.cpp
--- a/Engine/Plugins/Enterprise/DatasmithContent/Source/DatasmithContent/Private/ObjectTemplates/DatasmithObjectTemplate.cpp
+++ b/Engine/Plugins/Enterprise/DatasmithContent/Source/DatasmithContent/Private/ObjectTemplates/DatasmithObjectTemplate.cpp
@@ -117,6 +117,7 @@ void FDatasmithObjectTemplateUtils::SetObjectTemplate(UObject* Outer, UDatasmith
 	{
 		TMap< TSubclassOf< UDatasmithObjectTemplate >, UDatasmithObjectTemplate* >& ObjectTemplatesMap = UserData->ObjectTemplates;
 		ObjectTemplatesMap.FindOrAdd(ObjectTemplate->GetClass()) = ObjectTemplate;
+		ObjectTemplate->SetFlags(RF_Public);
 
 		if (ObjectTemplate->GetOuter() != UserData)
 		{
