Merging 6942139 from //UE4/Dev-Enterprise/... to  //UE4/Release-4.23/...

Fixing some references issue with the actors. The fix is simple now that the templates have the proper outer chain.
Also fixed a untracked issue with the reset override on Cine Camera Actor.

diff --git a/Engine/Plugins/Enterprise/DatasmithContent/Source/DatasmithContent/Private/ObjectTemplates/DatasmithCineCameraActorTemplate.cpp b/Engine/Plugins/Enterprise/DatasmithContent/Source/DatasmithContent/Private/ObjectTemplates/DatasmithCineCameraActorTemplate.cpp
--- a/Engine/Plugins/Enterprise/DatasmithContent/Source/DatasmithContent/Private/ObjectTemplates/DatasmithCineCameraActorTemplate.cpp
+++ b/Engine/Plugins/Enterprise/DatasmithContent/Source/DatasmithContent/Private/ObjectTemplates/DatasmithCineCameraActorTemplate.cpp
@@ -2,6 +2,8 @@
 
 #include "ObjectTemplates/DatasmithCineCameraActorTemplate.h"
 
+#include "ObjectTemplates/DatasmithActorTemplate.h"
+
 FDatasmithCameraLookatTrackingSettingsTemplate::FDatasmithCameraLookatTrackingSettingsTemplate()
 	: bEnableLookAtTracking( 0 )
 {
@@ -29,7 +31,7 @@ bool FDatasmithCameraLookatTrackingSettingsTemplate::Equals( const FDatasmithCam
 
 UObject* UDatasmithCineCameraActorTemplate::UpdateObject( UObject* Destination, bool bForce )
 {
-	ACineCameraActor* CineCameraActor = Cast< ACineCameraActor >( Destination );
+	ACineCameraActor* CineCameraActor = UDatasmithActorTemplate::GetActor< ACineCameraActor >( Destination );
 
 	if ( !CineCameraActor )
 	{
