Merging 7039724 from //UE4/Dev-Enterprise/... to //UE4/Release-4.23/...

Fixed an issue with the static mesh template leaving some non usefull redirectors. Those redirectors would also block the user from deleting the created static mesh.

diff --git a/Engine/Plugins/Enterprise/DatasmithContent/Source/DatasmithContent/Private/ObjectTemplates/DatasmithObjectTemplate.cpp b/Engine/Plugins/Enterprise/DatasmithContent/Source/DatasmithContent/Private/ObjectTemplates/DatasmithObjectTemplate.cpp
--- a/Engine/Plugins/Enterprise/DatasmithContent/Source/DatasmithContent/Private/ObjectTemplates/DatasmithObjectTemplate.cpp
+++ b/Engine/Plugins/Enterprise/DatasmithContent/Source/DatasmithContent/Private/ObjectTemplates/DatasmithObjectTemplate.cpp
@@ -123,7 +123,7 @@ void FDatasmithObjectTemplateUtils::SetObjectTemplate(UObject* Outer, UDatasmith
 		{
 			// The outer chain is important for most of the engine functionality.
 			// If it's not set properly the deep copy of object won't work properly
-			ObjectTemplate->Rename(nullptr, UserData);
+			ObjectTemplate->Rename(nullptr, UserData, REN_DontCreateRedirectors | REN_ForceNoResetLoaders);
 		}
 	}
 #endif // #if WITH_EDITORONLY_DATA
