Copied cl 6913024 6913207 6913218 6913465 6914208 6914620 6938884 6975365 6979856 from //UE4/Dev-Enterprise

- C4D: Fix instance names not being parsed correctly
- C4D: Fix bUseTemperature being set when it wasn't supposed to #jira UE-74081 #rb Anousack.Kitisa
- C4D: Import rectangle lights as Rect area light type #jira UE-72497 #rb Anousack.Kitisa
- Don't move camera focus for thumbnail if Bounds is invalid #jira UE-72794 #rb JeanMichel.Dignard
- C4D: Combine all reflectance and color channel's colors
- C4D: Fetch correct TextureTags when reimporting a static mesh
- Re-save asset for correct build.version
- C4D: Support for allowing camera roll when tracking actors
- Update import context folders to handle reimport with new file

#rnx

diff --git a/Engine/Plugins/Enterprise/DatasmithContent/Source/DatasmithContent/Private/ObjectTemplates/DatasmithCineCameraActorTemplate.cpp b/Engine/Plugins/Enterprise/DatasmithContent/Source/DatasmithContent/Private/ObjectTemplates/DatasmithCineCameraActorTemplate.cpp
--- a/Engine/Plugins/Enterprise/DatasmithContent/Source/DatasmithContent/Private/ObjectTemplates/DatasmithCineCameraActorTemplate.cpp
+++ b/Engine/Plugins/Enterprise/DatasmithContent/Source/DatasmithContent/Private/ObjectTemplates/DatasmithCineCameraActorTemplate.cpp
@@ -6,24 +6,28 @@
 
 FDatasmithCameraLookatTrackingSettingsTemplate::FDatasmithCameraLookatTrackingSettingsTemplate()
 	: bEnableLookAtTracking( 0 )
+	, bAllowRoll( 0 )
 {
 }
 
 void FDatasmithCameraLookatTrackingSettingsTemplate::Apply( FCameraLookatTrackingSettings* Destination, FDatasmithCameraLookatTrackingSettingsTemplate* PreviousTemplate )
 {
 	DATASMITHOBJECTTEMPLATE_CONDITIONALSET( bEnableLookAtTracking, Destination, PreviousTemplate );
+	DATASMITHOBJECTTEMPLATE_CONDITIONALSET( bAllowRoll, Destination, PreviousTemplate );
 	DATASMITHOBJECTTEMPLATE_CONDITIONALSETSOFTOBJECTPTR( ActorToTrack, Destination, PreviousTemplate );
 }
 
 void FDatasmithCameraLookatTrackingSettingsTemplate::Load( const FCameraLookatTrackingSettings& Source )
 {
 	bEnableLookAtTracking = Source.bEnableLookAtTracking;
+	bAllowRoll = Source.bAllowRoll;
 	ActorToTrack = Source.ActorToTrack;
 }
 
 bool FDatasmithCameraLookatTrackingSettingsTemplate::Equals( const FDatasmithCameraLookatTrackingSettingsTemplate& Other ) const
 {
 	bool bEquals = ( bEnableLookAtTracking == Other.bEnableLookAtTracking );
+	bEquals = bEquals && ( bAllowRoll == Other.bAllowRoll );
 	bEquals = bEquals && ( ActorToTrack == Other.ActorToTrack );
 
 	return bEquals;
diff --git a/Engine/Plugins/Enterprise/DatasmithContent/Source/DatasmithContent/Public/ObjectTemplates/DatasmithCineCameraActorTemplate.h b/Engine/Plugins/Enterprise/DatasmithContent/Source/DatasmithContent/Public/ObjectTemplates/DatasmithCineCameraActorTemplate.h
--- a/Engine/Plugins/Enterprise/DatasmithContent/Source/DatasmithContent/Public/ObjectTemplates/DatasmithCineCameraActorTemplate.h
+++ b/Engine/Plugins/Enterprise/DatasmithContent/Source/DatasmithContent/Public/ObjectTemplates/DatasmithCineCameraActorTemplate.h
@@ -17,6 +17,9 @@ public:
 	UPROPERTY()
 	uint8 bEnableLookAtTracking : 1;
 
+	UPROPERTY()
+	uint8 bAllowRoll : 1;
+
 	UPROPERTY()
 	TSoftObjectPtr< AActor > ActorToTrack;
 
