Copied 7062302 6979856 7062291 from dev-enterprise

- Update import context folders to handle reimport with new file
- Prevent SceneComponentTemplate from setting a component to static if its AttachParent is dynamic
- C4D: Don't use IES file brightness; Check when IES is enabled in C4D

#rnx

diff --git a/Engine/Plugins/Enterprise/DatasmithContent/Source/DatasmithContent/Private/ObjectTemplates/DatasmithSceneComponentTemplate.cpp b/Engine/Plugins/Enterprise/DatasmithContent/Source/DatasmithContent/Private/ObjectTemplates/DatasmithSceneComponentTemplate.cpp
--- a/Engine/Plugins/Enterprise/DatasmithContent/Source/DatasmithContent/Private/ObjectTemplates/DatasmithSceneComponentTemplate.cpp
+++ b/Engine/Plugins/Enterprise/DatasmithContent/Source/DatasmithContent/Private/ObjectTemplates/DatasmithSceneComponentTemplate.cpp
@@ -29,7 +29,14 @@ UObject* UDatasmithSceneComponentTemplate::UpdateObject( UObject* Destination, b
 
 	if ( !PreviousTemplate || PreviousTemplate->Mobility == SceneComponent->Mobility )
 	{
-		SceneComponent->SetMobility( Mobility );
+		// We can't attach a static component to a dynamic parent
+		EComponentMobility::Type TargetMobility = Mobility;
+		if (AttachParent)
+		{
+			TargetMobility = (EComponentMobility::Type) FMath::Max((uint8)Mobility, (uint8)AttachParent->Mobility);
+		}
+
+		SceneComponent->SetMobility( TargetMobility );
 	}
 
 	const ULevel* SceneComponentLevel = SceneComponent->GetComponentLevel();
