Datasmith:
- Improve UX for retesselation action

diff --git a/Engine/Plugins/Enterprise/DatasmithContent/Source/DatasmithContent/Public/DatasmithCustomAction.h b/Engine/Plugins/Enterprise/DatasmithContent/Source/DatasmithContent/Public/DatasmithCustomAction.h
--- a/Engine/Plugins/Enterprise/DatasmithContent/Source/DatasmithContent/Public/DatasmithCustomAction.h
+++ b/Engine/Plugins/Enterprise/DatasmithContent/Source/DatasmithContent/Public/DatasmithCustomAction.h
@@ -14,8 +14,8 @@ class IDatasmithCustomAction
 public:
 	virtual ~IDatasmithCustomAction() = default;
 
-	virtual const FText& GetLabel() { return FText::GetEmpty(); };
-	virtual const FText& GetTooltip() { return FText::GetEmpty(); };
+	virtual const FText& GetLabel() = 0;
+	virtual const FText& GetTooltip() = 0;
 
 	virtual bool CanApplyOnAssets(const TArray<FAssetData>& SelectedAssets) { return false; }
 	virtual void ApplyOnAssets(const TArray<FAssetData>& SelectedAssets) {}
@@ -29,9 +29,8 @@ class DATASMITHCONTENT_API UDatasmithCustomActionBase : public UObject, public I
 	GENERATED_BODY()
 
 public:
-
-#if WITH_EDITORONLY_DATA
-#endif
+	virtual const FText& GetLabel() override PURE_VIRTUAL(GetLabel, return FText::GetEmpty(););
+	virtual const FText& GetTooltip() override PURE_VIRTUAL(GetTooltip, return FText::GetEmpty(););
 };
 
 
