Copied 4348902 JeanMichel.Dignard
Datasmith UE Pbr
- Added support for Falloff of type perpendicular/parallel
- Added support to disable the dof on physical cameras based on max setting
- Parse the texture output for any texmap
- better support for map weights, blend to the proper color

diff --git a/Engine/Plugins/Enterprise/DatasmithContent/Source/DatasmithContent/Private/ObjectTemplates/DatasmithCineCameraComponentTemplate.cpp b/Engine/Plugins/Enterprise/DatasmithContent/Source/DatasmithContent/Private/ObjectTemplates/DatasmithCineCameraComponentTemplate.cpp
--- a/Engine/Plugins/Enterprise/DatasmithContent/Source/DatasmithContent/Private/ObjectTemplates/DatasmithCineCameraComponentTemplate.cpp
+++ b/Engine/Plugins/Enterprise/DatasmithContent/Source/DatasmithContent/Private/ObjectTemplates/DatasmithCineCameraComponentTemplate.cpp
@@ -2,8 +2,6 @@
 
 #include "ObjectTemplates/DatasmithCineCameraComponentTemplate.h"
 
-#include "CineCameraComponent.h"
-
 void FDatasmithCameraFilmbackSettingsTemplate::Apply( FCameraFilmbackSettings* Destination, const FDatasmithCameraFilmbackSettingsTemplate* PreviousTemplate )
 {
 	DATASMITHOBJECTTEMPLATE_CONDITIONALSET( SensorWidth, Destination, PreviousTemplate );
@@ -113,17 +111,20 @@ bool FDatasmithCameraLensSettingsTemplate::Equals( const FDatasmithCameraLensSet
 
 void FDatasmithCameraFocusSettingsTemplate::Apply( FCameraFocusSettings* Destination, const FDatasmithCameraFocusSettingsTemplate* PreviousTemplate )
 {
+	DATASMITHOBJECTTEMPLATE_CONDITIONALSET( FocusMethod, Destination, PreviousTemplate );
 	DATASMITHOBJECTTEMPLATE_CONDITIONALSET( ManualFocusDistance, Destination, PreviousTemplate );
 }
 
 void FDatasmithCameraFocusSettingsTemplate::Load( const FCameraFocusSettings& Source )
 {
+	FocusMethod = Source.FocusMethod;
 	ManualFocusDistance = Source.ManualFocusDistance;
 }
 
 bool FDatasmithCameraFocusSettingsTemplate::Equals( const FDatasmithCameraFocusSettingsTemplate& Other ) const
 {
-	bool bEquals = FMath::IsNearlyEqual( ManualFocusDistance, Other.ManualFocusDistance);
+	bool bEquals = ( FocusMethod == Other.FocusMethod );
+	bEquals = bEquals && FMath::IsNearlyEqual( ManualFocusDistance, Other.ManualFocusDistance);
 
 	return bEquals;
 }
diff --git a/Engine/Plugins/Enterprise/DatasmithContent/Source/DatasmithContent/Public/ObjectTemplates/DatasmithCineCameraComponentTemplate.h b/Engine/Plugins/Enterprise/DatasmithContent/Source/DatasmithContent/Public/ObjectTemplates/DatasmithCineCameraComponentTemplate.h
--- a/Engine/Plugins/Enterprise/DatasmithContent/Source/DatasmithContent/Public/ObjectTemplates/DatasmithCineCameraComponentTemplate.h
+++ b/Engine/Plugins/Enterprise/DatasmithContent/Source/DatasmithContent/Public/ObjectTemplates/DatasmithCineCameraComponentTemplate.h
@@ -2,6 +2,7 @@
 
 #pragma once
 
+#include "CineCameraComponent.h"
 #include "Engine/Scene.h"
 #include "ObjectTemplates/DatasmithObjectTemplate.h"
 
@@ -45,6 +46,9 @@ struct FDatasmithCameraFocusSettingsTemplate
 {
 	GENERATED_BODY()
 
+	UPROPERTY()
+	ECameraFocusMethod FocusMethod;
+
 	UPROPERTY()
 	float ManualFocusDistance;
 
