Copied 4350470 Johan.Duparc
Fix crash on property edition for blueprint actors with templated asset user data

diff --git a/Engine/Plugins/Enterprise/DatasmithContent/Source/DatasmithContent/Private/ObjectTemplates/DatasmithObjectTemplate.cpp b/Engine/Plugins/Enterprise/DatasmithContent/Source/DatasmithContent/Private/ObjectTemplates/DatasmithObjectTemplate.cpp
--- a/Engine/Plugins/Enterprise/DatasmithContent/Source/DatasmithContent/Private/ObjectTemplates/DatasmithObjectTemplate.cpp
+++ b/Engine/Plugins/Enterprise/DatasmithContent/Source/DatasmithContent/Private/ObjectTemplates/DatasmithObjectTemplate.cpp
@@ -63,7 +63,8 @@ TMap< TSubclassOf< UDatasmithObjectTemplate >, UDatasmithObjectTemplate* >* FDat
 
 	if (!UserData)
 	{
-		UserData = NewObject< UDatasmithAssetUserData >(Outer, NAME_None, RF_Public | RF_Transactional);
+		EObjectFlags Flags = RF_Public /*| RF_Transactional*/; // RF_Transactional Disabled as is can cause a crash in the transaction system for blueprints
+		UserData = NewObject< UDatasmithAssetUserData >(Outer, NAME_None, Flags);
 		AssetUserDataInterface->AddAssetUserData(UserData);
 	}
 
