Copied 43737703 Johan.Duparc
Prevent crash for old datasmith assets
Clear transactional flag in user asset data that could cause a crash for Blueprint assets.

diff --git a/Engine/Plugins/Enterprise/DatasmithContent/Source/DatasmithContent/Public/DatasmithAssetUserData.h b/Engine/Plugins/Enterprise/DatasmithContent/Source/DatasmithContent/Public/DatasmithAssetUserData.h
--- a/Engine/Plugins/Enterprise/DatasmithContent/Source/DatasmithContent/Public/DatasmithAssetUserData.h
+++ b/Engine/Plugins/Enterprise/DatasmithContent/Source/DatasmithContent/Public/DatasmithAssetUserData.h
@@ -25,4 +25,19 @@ public:
 	UPROPERTY()
 	TMap< TSubclassOf< class UDatasmithObjectTemplate >, UDatasmithObjectTemplate* > ObjectTemplates;
 #endif
+
+protected:
+	virtual bool IsPostLoadThreadSafe() const override
+	{
+		return true;
+	}
+
+	virtual void PostLoad() override
+	{
+		Super::PostLoad();
+
+		// Avoid possible crash with old data:
+		// RF_Transactional flag can cause a crash on serialization for Blueprint instances
+		ClearFlags(RF_Transactional);
+	}
 };
