Copied 4374926 & 4377431 Johan.Duparc
Fix crash on reimport with removed AreaLight

diff --git a/Engine/Plugins/Enterprise/DatasmithContent/Source/DatasmithContent/Private/DatasmithAssetUserData.cpp b/Engine/Plugins/Enterprise/DatasmithContent/Source/DatasmithContent/Private/DatasmithAssetUserData.cpp
new file mode 100644
--- /dev/null
+++ b/Engine/Plugins/Enterprise/DatasmithContent/Source/DatasmithContent/Private/DatasmithAssetUserData.cpp
@@ -0,0 +1,34 @@
+// Copyright 1998-2018 Epic Games, Inc. All Rights Reserved.
+
+#include "DatasmithAssetUserData.h"
+
+#include "DatasmithContentModule.h"
+
+#if WITH_EDITORONLY_DATA
+
+bool UDatasmithAssetUserData::IsPostLoadThreadSafe() const
+{
+	return true;
+}
+
+void UDatasmithAssetUserData::PostLoad()
+{
+	Super::PostLoad();
+
+	// RF_Transactional flag can cause a crash on save for Blueprint instances, and old data was flagged.
+	ClearFlags(RF_Transactional);
+
+	// A serialization issue caused nullptr to be serialized instead of valid UDatasmithObjectTemplate pointers.
+	// This cleanup ensure values from this map can always be dereferenced
+	for (auto It = ObjectTemplates.CreateIterator(); It; ++It)
+	{
+		if (!ensure(It->Value))
+		{
+			It.RemoveCurrent();
+			UE_LOG(LogDatasmithContent, Warning, TEXT("Serialization issue: null value found in templates"))
+		}
+	}
+}
+
+#endif // WITH_EDITORONLY_DATA
+
diff --git a/Engine/Plugins/Enterprise/DatasmithContent/Source/DatasmithContent/Private/DatasmithContentModule.cpp b/Engine/Plugins/Enterprise/DatasmithContent/Source/DatasmithContent/Private/DatasmithContentModule.cpp
--- a/Engine/Plugins/Enterprise/DatasmithContent/Source/DatasmithContent/Private/DatasmithContentModule.cpp
+++ b/Engine/Plugins/Enterprise/DatasmithContent/Source/DatasmithContent/Private/DatasmithContentModule.cpp
@@ -10,6 +10,8 @@
 #include "Misc/Paths.h"
 #include "ShaderCore.h"
 
+DEFINE_LOG_CATEGORY(LogDatasmithContent)
+
 /**
  * DatasmithContent module implementation (private)
  */
diff --git a/Engine/Plugins/Enterprise/DatasmithContent/Source/DatasmithContent/Public/DatasmithAssetUserData.h b/Engine/Plugins/Enterprise/DatasmithContent/Source/DatasmithContent/Public/DatasmithAssetUserData.h
--- a/Engine/Plugins/Enterprise/DatasmithContent/Source/DatasmithContent/Public/DatasmithAssetUserData.h
+++ b/Engine/Plugins/Enterprise/DatasmithContent/Source/DatasmithContent/Public/DatasmithAssetUserData.h
@@ -24,20 +24,8 @@ public:
 #if WITH_EDITORONLY_DATA
 	UPROPERTY()
 	TMap< TSubclassOf< class UDatasmithObjectTemplate >, UDatasmithObjectTemplate* > ObjectTemplates;
-#endif
-
-protected:
-	virtual bool IsPostLoadThreadSafe() const override
-	{
-		return true;
-	}
 
-	virtual void PostLoad() override
-	{
-		Super::PostLoad();
-
-		// Avoid possible crash with old data:
-		// RF_Transactional flag can cause a crash on serialization for Blueprint instances
-		ClearFlags(RF_Transactional);
-	}
+	virtual bool IsPostLoadThreadSafe() const override;
+	virtual void PostLoad() override;
+#endif
 };
diff --git a/Engine/Plugins/Enterprise/DatasmithContent/Source/DatasmithContent/Public/DatasmithContentModule.h b/Engine/Plugins/Enterprise/DatasmithContent/Source/DatasmithContent/Public/DatasmithContentModule.h
--- a/Engine/Plugins/Enterprise/DatasmithContent/Source/DatasmithContent/Public/DatasmithContentModule.h
+++ b/Engine/Plugins/Enterprise/DatasmithContent/Source/DatasmithContent/Public/DatasmithContentModule.h
@@ -4,9 +4,11 @@
 
 #include "Modules/ModuleInterface.h"
 #include "Modules/ModuleManager.h"		// For inline LoadModuleChecked()
+#include "Logging/LogMacros.h"
 
 #define DATASMITHCONTENT_MODULE_NAME TEXT("DatasmithContent")
 
+DECLARE_LOG_CATEGORY_EXTERN( LogDatasmithContent, Log, All );
 DECLARE_DELEGATE_OneParam( FOnSpawnDatasmithSceneActors, class ADatasmithSceneActor* );
 
 /**
