Copied 4392578 JeanMichel.Dignard
Fixed actors not saving properly on a reimport when moved to a sublevel.
The issue is that they were being parented to actors from another level which isn't supported.
Also, the scene component template should refer to the parent component through a TSoftObjectPtr in case it's in another package.

diff --git a/Engine/Plugins/Enterprise/DatasmithContent/Source/DatasmithContent/Private/ObjectTemplates/DatasmithSceneComponentTemplate.cpp b/Engine/Plugins/Enterprise/DatasmithContent/Source/DatasmithContent/Private/ObjectTemplates/DatasmithSceneComponentTemplate.cpp
--- a/Engine/Plugins/Enterprise/DatasmithContent/Source/DatasmithContent/Private/ObjectTemplates/DatasmithSceneComponentTemplate.cpp
+++ b/Engine/Plugins/Enterprise/DatasmithContent/Source/DatasmithContent/Private/ObjectTemplates/DatasmithSceneComponentTemplate.cpp
@@ -30,6 +30,8 @@ void UDatasmithSceneComponentTemplate::Apply( UObject* Destination, bool bForce
 		SceneComponent->SetMobility( Mobility );
 	}
 
+	bool bCanAttach = ( AttachParent && AttachParent->GetComponentLevel() == SceneComponent->GetComponentLevel() );
+
 	if ( !PreviousTemplate || PreviousTemplate->AttachParent == SceneComponent->GetAttachParent() )
 	{
 		FAttachmentTransformRules AttachmentTransformRules = FAttachmentTransformRules::KeepRelativeTransform;
@@ -44,14 +46,31 @@ void UDatasmithSceneComponentTemplate::Apply( UObject* Destination, bool bForce
 			AttachmentTransformRules = FAttachmentTransformRules::KeepWorldTransform;
 		}
 
-		SceneComponent->AttachToComponent( AttachParent, AttachmentTransformRules );
+		if ( bCanAttach )
+		{
+			SceneComponent->AttachToComponent( AttachParent.Get(), AttachmentTransformRules );
+		}
 	}
 
 	if ( !PreviousTemplate || AreTransformsEqual( PreviousTemplate->RelativeTransform, SceneComponent->GetRelativeTransform() ) )
 	{
-		SceneComponent->SetRelativeTransform( RelativeTransform );
-	}
+		if ( bCanAttach )
+		{
+			SceneComponent->SetRelativeTransform( RelativeTransform );
+		}
+		else
+		{
+			// We were unable to attach to our parent so we need to compute the desired world transform
+			FTransform WorldTransform = RelativeTransform;
+
+			if ( AttachParent )
+			{
+				WorldTransform *= AttachParent->GetComponentTransform();
 
+				SceneComponent->SetRelativeTransform( WorldTransform );
+			}
+		}
+	}
 
 	if ( !PreviousTemplate )
 	{
diff --git a/Engine/Plugins/Enterprise/DatasmithContent/Source/DatasmithContent/Public/ObjectTemplates/DatasmithSceneComponentTemplate.h b/Engine/Plugins/Enterprise/DatasmithContent/Source/DatasmithContent/Public/ObjectTemplates/DatasmithSceneComponentTemplate.h
--- a/Engine/Plugins/Enterprise/DatasmithContent/Source/DatasmithContent/Public/ObjectTemplates/DatasmithSceneComponentTemplate.h
+++ b/Engine/Plugins/Enterprise/DatasmithContent/Source/DatasmithContent/Public/ObjectTemplates/DatasmithSceneComponentTemplate.h
@@ -21,7 +21,7 @@ public:
 	TEnumAsByte< EComponentMobility::Type > Mobility;
 
 	UPROPERTY()
-	class USceneComponent* AttachParent;
+	TSoftObjectPtr< USceneComponent > AttachParent;
 
 	UPROPERTY()
 	TSet<FName> Tags;
