Copied 4500498 Alin.Loghin

Compilation fix for AxF carpaint shader on vulkan

diff --git a/Engine/Plugins/Enterprise/DatasmithContent/Shaders/Private/FlakesAccess.ush b/Engine/Plugins/Enterprise/DatasmithContent/Shaders/Private/FlakesAccess.ush
--- a/Engine/Plugins/Enterprise/DatasmithContent/Shaders/Private/FlakesAccess.ush
+++ b/Engine/Plugins/Enterprise/DatasmithContent/Shaders/Private/FlakesAccess.ush
@@ -1,4 +1,3 @@
-
 /**
 Flakes texture access based on light angular direction, uses a pseduo(i.e. in atlas format) 2D texture array
 to access flakes slices and finally applies bilinear interpolation.
@@ -19,52 +18,46 @@ to access flakes slices and finally applies bilinear interpolation.
 #define PI2 1.5707963267948966
 #define USE_GRADIENTS 1
 
-struct TextureArraySample
-{ 
-    float4 Sample(float2 UV, float Index)
-    {
-      float SideCount = sqrt(FlakesTexCount); // per side texture count
-      float InvSideCount = 1.0 / SideCount; 
-      float TexPadding = 8.0; // pixels padding count
-      float TexSize = FlakesTexValues.x;
-      float PaddedTexSize = TexSize + TexPadding * 2.0;
-      float AtlasTexelSize = 1.0 / (SideCount * PaddedTexSize);
-
-      float2 TexOffset = Index;
-      TexOffset.x = fmod(TexOffset.x, SideCount);
-      TexOffset.y = floor(TexOffset.y * InvSideCount);
-      TexOffset *= InvSideCount;
-      TexOffset += TexPadding * AtlasTexelSize;
+// NOTE. had to use macros as function structs aren't portable on certain platforms such as Vulkan
 
-      float2 SampleTexCoord = frac(UV * Tiling);
-      SampleTexCoord *= InvSideCount * (TexSize / PaddedTexSize);
-      SampleTexCoord += TexOffset;
+#if USE_GRADIENTS == 1 
+#define SAMPLE_TEX_COORD(SampleTexCoord, Result) \
+    float2 DdxUV = DDX(SampleTexCoord); \
+    float2 DdyUV = DDY(SampleTexCoord); \
+    Result = Texture2DSampleGrad(FlakesTexArray, FlakesTexArraySampler, SampleTexCoord, DdxUV, DdyUV); 
+#else 
+    Result = Texture2DSample(FlakesTexArray, FlakesTexArraySampler, SampleTexCoord); 
+#endif 
 
-#if USE_GRADIENTS == 1
-      float2 DdxUV = DDX(SampleTexCoord);
-      float2 DdyUV = DDY(SampleTexCoord);
-      return Texture2DSampleGrad(FlakesTexArray, FlakesTexArraySampler, SampleTexCoord, DdxUV, DdyUV);
-#else
-      return Texture2DSample(FlakesTexArray, FlakesTexArraySampler, SampleTexCoord);
-#endif
-    }  
-}; 
-
-struct ThetaSliceLUTSample
-{ 
-    float Sample(float BinThetaI)
-    {
-      float TexelSize = FlakesTexValues.y;
-      float Padding = 4.0;
-      float Index = BinThetaI * Padding * TexelSize + TexelSize;
+#define SAMPLE_TEX_ARRAY(UV, Index, Result) \
+{ \
+    float SideCount = sqrt(FlakesTexCount); \
+    float InvSideCount = 1.0 / SideCount;  \
+    float TexPadding = 8.0; \
+    float TexSize = FlakesTexValues.x; \
+    float PaddedTexSize = TexSize + TexPadding * 2.0; \
+    float AtlasTexelSize = 1.0 / (SideCount * PaddedTexSize); \
+ \
+    float2 TexOffset = Index; \
+    TexOffset.x = fmod(TexOffset.x, SideCount); \
+    TexOffset.y = floor(TexOffset.y * InvSideCount); \
+    TexOffset *= InvSideCount; \
+    TexOffset += TexPadding * AtlasTexelSize; \
+ \
+    float2 SampleTexCoord = frac(UV * Tiling); \
+    SampleTexCoord *= InvSideCount * (TexSize / PaddedTexSize); \
+    SampleTexCoord += TexOffset; \
+    SAMPLE_TEX_COORD(SampleTexCoord, Result); \
+}   
 
-      float Value = Texture2DSampleLevel(ThetaSliceLUT, ThetaSliceLUTSampler, float2(Index, 0.0), 0.0).x;
-      Value = floor(Value * 255.0 + 0.5);
-      return Value;
-    }
-};
-TextureArraySample TextureArray;
-ThetaSliceLUTSample ThetaFiLUT;
+#define SAMPLE_LUT(BinThetaI, Result) \
+{ \
+    float TexelSize = FlakesTexValues.y; \
+    float Padding = 4.0; \
+    float Index = (BinThetaI) * Padding * TexelSize + TexelSize; \
+    float Value = Texture2DSampleLevel(ThetaSliceLUT, ThetaSliceLUTSampler, float2(Index, 0.0), 0.0).x; \
+    Result = floor(Value * 255.0 + 0.5); \
+}
 
 //ThetaF sampling defines the angular sampling, i.e. angular flake lifetime
 float ThetaF = clamp(ThetaAngles.x, 0.0, 1.0 / PI2) * NumThetaF + 0.5;
@@ -81,8 +74,8 @@ if (NumThetaI < NumThetaF)
 {
     if ((LowThetaI % 2) == 1 || (HighThetaI % 2) == 1) 
     { 
-      //tweak for extra randomness, in the case of almost planar geometries with repeating patches
-      UV.xy = UV.yx;
+        //tweak for extra randomness, in the case of almost planar geometries with repeating patches
+        UV.xy = UV.yx;
     }
 
     //map to the original sampling
@@ -98,31 +91,40 @@ half3 ColorTopRight  = 0.0;
 //access flakes texture and stay in the correct slices(no slip over)
 if (LowThetaI < MaxThetaI)
 {
-    float LowThetaIndex = ThetaFiLUT.Sample(LowThetaI) + LowThetaF;
-    if (LowThetaIndex < ThetaFiLUT.Sample(LowThetaI + 1.0))
+    float BaseThetaIndex, UpperIndexLimit;
+    SAMPLE_LUT(LowThetaI, BaseThetaIndex);
+    SAMPLE_LUT(LowThetaI + 1.0, UpperIndexLimit);
+
+    float LowThetaIndex = BaseThetaIndex + LowThetaF;
+    if (LowThetaIndex < UpperIndexLimit)
     {
-        ColorBottomLeft = TextureArray.Sample(UV, LowThetaIndex).xyz;
+        SAMPLE_TEX_ARRAY(UV, LowThetaIndex, ColorBottomLeft.xyz);
     }
-    float HighThetaIndex = ThetaFiLUT.Sample(LowThetaI) + HighThetaF;
-    if (HighThetaIndex < ThetaFiLUT.Sample(LowThetaI + 1))
+    float HighThetaIndex = BaseThetaIndex + HighThetaF;
+    if (HighThetaIndex < UpperIndexLimit)
     {
-        ColorBottomRight = TextureArray.Sample(UV, HighThetaIndex).xyz;
+        SAMPLE_TEX_ARRAY(UV, HighThetaIndex, ColorBottomRight.xyz);
     }
 }
 
 if (HighThetaI < MaxThetaI)
 {
-    float LowThetaIndex = ThetaFiLUT.Sample(HighThetaI) + LowThetaF;
-    if (LowThetaIndex < ThetaFiLUT.Sample(HighThetaI + 1))
+    float BaseThetaIndex, UpperIndexLimit;
+    SAMPLE_LUT(HighThetaI, BaseThetaIndex);
+    SAMPLE_LUT(HighThetaI + 1.0, UpperIndexLimit);
+
+    float LowThetaIndex = BaseThetaIndex + LowThetaF;
+    if (LowThetaIndex < UpperIndexLimit)
     {
-        ColorTopLeft = TextureArray.Sample(UV, LowThetaIndex).xyz;
+        SAMPLE_TEX_ARRAY(UV, LowThetaIndex, ColorTopLeft.xyz);
     }
-    float HighThetaIndex = ThetaFiLUT.Sample(HighThetaI) + HighThetaF;
-    if (HighThetaIndex < ThetaFiLUT.Sample(HighThetaI + 1))
+    float HighThetaIndex = BaseThetaIndex + HighThetaF;
+    if (HighThetaIndex < UpperIndexLimit)
     {
-        ColorTopRight = TextureArray.Sample(UV, HighThetaIndex).xyz;
+        SAMPLE_TEX_ARRAY(UV, HighThetaIndex, ColorTopRight.xyz);
     }
 }
+
 //bilinear interpolation
 float DiffThetaF = ThetaF - float(LowThetaF);
 float DiffThetaI = ThetaI - float(LowThetaI);
