Copied cl 4799858 & 4799861 from dev-enterprise

Exposed spotlight angles and specular scale on the Datasmith Area Light blueprint actor.
Datasmith master materials enabling usage with static lighting and instanced static meshes

diff --git a/Engine/Plugins/Enterprise/DatasmithContent/Source/DatasmithContent/Private/DatasmithAreaLightActor.cpp b/Engine/Plugins/Enterprise/DatasmithContent/Source/DatasmithContent/Private/DatasmithAreaLightActor.cpp
--- a/Engine/Plugins/Enterprise/DatasmithContent/Source/DatasmithContent/Private/DatasmithAreaLightActor.cpp
+++ b/Engine/Plugins/Enterprise/DatasmithContent/Source/DatasmithContent/Private/DatasmithAreaLightActor.cpp
@@ -16,5 +16,7 @@ ADatasmithAreaLightActor::ADatasmithAreaLightActor()
 	, SourceRadius( 0.f )
 	, SourceLength( 0.f )
 	, AttenuationRadius( 1000.f )
+	, SpotlightInnerAngle(1.f)
+	, SpotlightOuterAngle(44.f)
 {
 }
\ No newline at end of file
diff --git a/Engine/Plugins/Enterprise/DatasmithContent/Source/DatasmithContent/Public/DatasmithAreaLightActor.h b/Engine/Plugins/Enterprise/DatasmithContent/Source/DatasmithContent/Public/DatasmithAreaLightActor.h
--- a/Engine/Plugins/Enterprise/DatasmithContent/Source/DatasmithContent/Public/DatasmithAreaLightActor.h
+++ b/Engine/Plugins/Enterprise/DatasmithContent/Source/DatasmithContent/Public/DatasmithAreaLightActor.h
@@ -74,6 +74,12 @@ public:
 	UPROPERTY(EditAnywhere, BlueprintReadWrite, Category = "Light", meta=(UIMin = "8.0", UIMax = "16384.0", SliderExponent = "5.0"))
 	float AttenuationRadius;
 
+	UPROPERTY(EditAnywhere, BlueprintReadWrite, Category = "Light", meta = (UIMin = "1.0", UIMax = "80.0", SliderExponent = "1.0"))
+	float SpotlightInnerAngle;
+
+	UPROPERTY(EditAnywhere, BlueprintReadWrite, Category = "Light", meta = (UIMin = "0.0", UIMax = "80.0", SliderExponent = "1.0"))
+	float SpotlightOuterAngle;
+
 public:
 	ADatasmithAreaLightActor();
 };
\ No newline at end of file
diff --git a/Engine/Source/Runtime/Engine/Classes/Components/LightComponent.h b/Engine/Source/Runtime/Engine/Classes/Components/LightComponent.h
--- a/Engine/Source/Runtime/Engine/Classes/Components/LightComponent.h
+++ b/Engine/Source/Runtime/Engine/Classes/Components/LightComponent.h
@@ -276,6 +276,9 @@ public:
 	UFUNCTION(BlueprintCallable, Category="Rendering|Components|Light")
 	void SetShadowBias(float NewValue);
 
+	UFUNCTION(BlueprintCallable, Category = "Rendering|Components|Light")
+	void SetSpecularScale(float NewValue);
+
 	UFUNCTION(BlueprintCallable, Category = "Rendering|Components|Light")
 	void SetForceCachedShadowsForMovablePrimitives(bool bNewValue);
 
diff --git a/Engine/Source/Runtime/Engine/Private/Components/LightComponent.cpp b/Engine/Source/Runtime/Engine/Private/Components/LightComponent.cpp
--- a/Engine/Source/Runtime/Engine/Private/Components/LightComponent.cpp
+++ b/Engine/Source/Runtime/Engine/Private/Components/LightComponent.cpp
@@ -981,6 +981,16 @@ void ULightComponent::SetShadowBias(float NewValue)
 	}
 }
 
+void ULightComponent::SetSpecularScale(float NewValue)
+{
+	if (AreDynamicDataChangesAllowed()
+		&& SpecularScale != NewValue)
+	{
+		SpecularScale = NewValue;
+		MarkRenderStateDirty();
+	}
+}
+
 void ULightComponent::SetForceCachedShadowsForMovablePrimitives(bool bNewValue)
 {
 	if (AreDynamicDataChangesAllowed()
