Merging 4845189 from Dev-Enterprise to Release-4.22

Fixed a crash that occur when changing the level editor active level (world) while a dataprep world was build.

The final world is now not used until the finalize actors.

diff --git a/Engine/Plugins/Enterprise/DatasmithContent/Source/DatasmithContentEditor/Public/DatasmithContentEditorModule.h b/Engine/Plugins/Enterprise/DatasmithContent/Source/DatasmithContentEditor/Public/DatasmithContentEditorModule.h
--- a/Engine/Plugins/Enterprise/DatasmithContent/Source/DatasmithContentEditor/Public/DatasmithContentEditorModule.h
+++ b/Engine/Plugins/Enterprise/DatasmithContent/Source/DatasmithContentEditor/Public/DatasmithContentEditorModule.h
@@ -21,8 +21,9 @@ public:
 	virtual ~IDataPrepImporterInterface() = default;
 
 	virtual UDatasmithImportOptions* CreateOptions(UObject* Outer) = 0;
-	virtual bool Initialize(const FString& GUID, UWorld* ImportWorld, UWorld* FinalWorld, UPackage* ParentPackage, UDatasmithScene* DatasmithScene, UDatasmithImportOptions* Options) = 0;
+	virtual bool Initialize(const FString& GUID, UWorld* ImportWorld, UPackage* ParentPackage, UDatasmithScene* DatasmithScene, UDatasmithImportOptions* Options) = 0;
 	virtual bool BuildWorld(TArray<TWeakObjectPtr<UObject>>& OutAssets) = 0;
+	virtual bool SetFinalWorld(UWorld* FinalWorld) = 0;
 	virtual bool FinalizeAssets(const TArray<TWeakObjectPtr<UObject>>& Assets) = 0;
 };
 
