Copied 4787074 from Dev-Enterprise

Disabled Datasmith animation import (also disabled in the Import Options UI) if geometry, light, camera imports are disabled.

diff --git a/Engine/Plugins/Enterprise/DatasmithContent/Source/DatasmithContent/Private/DatasmithAssetImportData.cpp b/Engine/Plugins/Enterprise/DatasmithContent/Source/DatasmithContent/Private/DatasmithAssetImportData.cpp
--- a/Engine/Plugins/Enterprise/DatasmithContent/Source/DatasmithContent/Private/DatasmithAssetImportData.cpp
+++ b/Engine/Plugins/Enterprise/DatasmithContent/Source/DatasmithContent/Private/DatasmithAssetImportData.cpp
@@ -14,6 +14,23 @@
 
 #if WITH_EDITORONLY_DATA
 
+#if WITH_EDITOR
+bool UDatasmithSceneImportData::CanEditChange(const UProperty* InProperty) const
+{
+	if (!Super::CanEditChange(InProperty))
+	{
+		return false;
+	}
+
+	if (InProperty->GetFName() == GET_MEMBER_NAME_CHECKED(FDatasmithImportBaseOptions, bIncludeAnimation))
+	{
+		return BaseOptions.CanIncludeAnimation();
+	}
+
+	return true;
+}
+#endif //WITH_EDITOR
+
 UDatasmithStaticMeshImportData* UDatasmithStaticMeshImportData::GetImportDataForStaticMesh( UStaticMesh* StaticMesh, TOptional< DefaultOptionsPair > DefaultImportOptions )
 {
 	check( StaticMesh );
diff --git a/Engine/Plugins/Enterprise/DatasmithContent/Source/DatasmithContent/Private/DatasmithImportOptions.cpp b/Engine/Plugins/Enterprise/DatasmithContent/Source/DatasmithContent/Private/DatasmithImportOptions.cpp
--- a/Engine/Plugins/Enterprise/DatasmithContent/Source/DatasmithContent/Private/DatasmithImportOptions.cpp
+++ b/Engine/Plugins/Enterprise/DatasmithContent/Source/DatasmithContent/Private/DatasmithImportOptions.cpp
@@ -109,6 +109,23 @@ void UDatasmithImportOptions::UpdateNotDisplayedConfig( bool bIsAReimport )
 	SearchPackagePolicy = EDatasmithImportSearchPackagePolicy::Current;
 }
 
+#if WITH_EDITOR
+bool UDatasmithImportOptions::CanEditChange(const UProperty* InProperty) const
+{
+	if (!Super::CanEditChange(InProperty))
+	{
+		return false;
+	}
+
+	if (InProperty->GetFName() == GET_MEMBER_NAME_CHECKED(FDatasmithImportBaseOptions, bIncludeAnimation))
+	{
+		return BaseOptions.CanIncludeAnimation();
+	}
+
+	return true;
+}
+#endif //WITH_EDITOR
+
 FDatasmithStaticMeshImportOptions::FDatasmithStaticMeshImportOptions()
 	: MinLightmapResolution( EDatasmithImportLightmapMin::LIGHTMAP_64 )
 	, MaxLightmapResolution( EDatasmithImportLightmapMax::LIGHTMAP_512 )
diff --git a/Engine/Plugins/Enterprise/DatasmithContent/Source/DatasmithContent/Public/DatasmithAssetImportData.h b/Engine/Plugins/Enterprise/DatasmithContent/Source/DatasmithContent/Public/DatasmithAssetImportData.h
--- a/Engine/Plugins/Enterprise/DatasmithContent/Source/DatasmithContent/Public/DatasmithAssetImportData.h
+++ b/Engine/Plugins/Enterprise/DatasmithContent/Source/DatasmithContent/Public/DatasmithAssetImportData.h
@@ -89,6 +89,12 @@ public:
 	UPROPERTY(EditAnywhere, Category = "Options", meta = (ShowOnlyInnerProperties))
 	FDatasmithImportBaseOptions BaseOptions;
 
+	//~ UObject interface
+#if WITH_EDITOR
+	virtual bool CanEditChange(const UProperty* InProperty) const override;
+#endif //WITH_EDITOR
+	//~ End UObject interface
+
 #endif // WITH_EDITORONLY_DATA
 };
 
diff --git a/Engine/Plugins/Enterprise/DatasmithContent/Source/DatasmithContent/Public/DatasmithImportOptions.h b/Engine/Plugins/Enterprise/DatasmithContent/Source/DatasmithContent/Public/DatasmithImportOptions.h
--- a/Engine/Plugins/Enterprise/DatasmithContent/Source/DatasmithContent/Public/DatasmithImportOptions.h
+++ b/Engine/Plugins/Enterprise/DatasmithContent/Source/DatasmithContent/Public/DatasmithImportOptions.h
@@ -203,6 +203,8 @@ struct DATASMITHCONTENT_API FDatasmithImportBaseOptions
 
 	UPROPERTY(config, EditAnywhere, BlueprintReadWrite, AdvancedDisplay, Category = Process, meta = (ShowOnlyInnerProperties))
 	FDatasmithStaticMeshImportOptions StaticMeshOptions;
+
+	bool CanIncludeAnimation() const { return bIncludeGeometry || bIncludeCamera || bIncludeLight; }
 };
 
 USTRUCT(BlueprintType)
@@ -326,4 +328,10 @@ public:
 	bool bUseSameOptions;
 
 	void UpdateNotDisplayedConfig( bool bIsAReimport );
+
+	//~ UObject interface
+#if WITH_EDITOR
+	virtual bool CanEditChange(const UProperty* InProperty) const override;
+#endif //WITH_EDITOR
+	//~ End UObject interface
 };
