Copied CL #4824630 from  //UE4/Dev-Enterprise/...

Fixed Datasmith reset overrides on the DatasmithAreaLightActor. The object passed to Apply is a scene component. The code was only handle the actor type.

diff --git a/Engine/Plugins/Enterprise/DatasmithContent/Source/DatasmithContent/Private/ObjectTemplates/DatasmithAreaLightActorTemplate.cpp b/Engine/Plugins/Enterprise/DatasmithContent/Source/DatasmithContent/Private/ObjectTemplates/DatasmithAreaLightActorTemplate.cpp
--- a/Engine/Plugins/Enterprise/DatasmithContent/Source/DatasmithContent/Private/ObjectTemplates/DatasmithAreaLightActorTemplate.cpp
+++ b/Engine/Plugins/Enterprise/DatasmithContent/Source/DatasmithContent/Private/ObjectTemplates/DatasmithAreaLightActorTemplate.cpp
@@ -12,7 +12,8 @@ UDatasmithAreaLightActorTemplate::UDatasmithAreaLightActorTemplate()
 void UDatasmithAreaLightActorTemplate::Apply( UObject* Destination, bool bForce )
 {
 #if WITH_EDITORONLY_DATA
-	ADatasmithAreaLightActor* AreaLightActor = Cast< ADatasmithAreaLightActor >( Destination );
+	const USceneComponent* SceneComponent = Cast< USceneComponent >( Destination );
+	ADatasmithAreaLightActor* AreaLightActor = Cast< ADatasmithAreaLightActor >( SceneComponent ? SceneComponent->GetOwner() : Destination );
 
 	if ( !AreaLightActor )
 	{
@@ -37,6 +38,8 @@ void UDatasmithAreaLightActorTemplate::Apply( UObject* Destination, bool bForce
 	DATASMITHOBJECTTEMPLATE_CONDITIONALSET( AttenuationRadius, AreaLightActor, PreviousTemplate );
 
 	FDatasmithObjectTemplateUtils::SetObjectTemplate( AreaLightActor->GetRootComponent(), this );
+	AreaLightActor->RerunConstructionScripts();
+
 #endif // #if WITH_EDITORONLY_DATA
 }
 
