Copied CL #4846980 from  //UE4/Dev-Enterprise/...

Fixed file paths being sanitized which made them point to nowhere.

diff --git a/Engine/Plugins/Enterprise/DatasmithContent/Source/DatasmithContent/Private/DatasmithAssetImportData.cpp b/Engine/Plugins/Enterprise/DatasmithContent/Source/DatasmithContent/Private/DatasmithAssetImportData.cpp
--- a/Engine/Plugins/Enterprise/DatasmithContent/Source/DatasmithContent/Private/DatasmithAssetImportData.cpp
+++ b/Engine/Plugins/Enterprise/DatasmithContent/Source/DatasmithContent/Private/DatasmithAssetImportData.cpp
@@ -33,7 +33,10 @@ bool UDatasmithSceneImportData::CanEditChange(const UProperty* InProperty) const
 
 UDatasmithStaticMeshImportData* UDatasmithStaticMeshImportData::GetImportDataForStaticMesh( UStaticMesh* StaticMesh, TOptional< DefaultOptionsPair > DefaultImportOptions )
 {
-	check( StaticMesh );
+	if( !StaticMesh )
+	{
+		return nullptr;
+	}
 
 	UDatasmithStaticMeshImportData* ImportData = nullptr;
 
