Fixed reimport of a dataprep imported mesh crashing
- The mesh outer was being set to a file path instead of a content path.
- Made sure the options used by the dataprep are coming from the Datasmith Scene.
- Made sure reimporting a single asset would retrieve the associated Datasmith Scene so that we can resolve dependencies properly.

diff --git a/Engine/Plugins/Enterprise/DatasmithContent/Source/DatasmithContentEditor/Public/DatasmithContentEditorModule.h b/Engine/Plugins/Enterprise/DatasmithContent/Source/DatasmithContentEditor/Public/DatasmithContentEditorModule.h
--- a/Engine/Plugins/Enterprise/DatasmithContent/Source/DatasmithContentEditor/Public/DatasmithContentEditorModule.h
+++ b/Engine/Plugins/Enterprise/DatasmithContent/Source/DatasmithContentEditor/Public/DatasmithContentEditorModule.h
@@ -9,6 +9,7 @@
 
 #define DATASMITHCONTENTEDITOR_MODULE_NAME TEXT("DatasmithContentEditor")
 
+struct FGuid;
 class UDatasmithImportOptions;
 class UDatasmithScene;
 class UPackage;
@@ -20,8 +21,12 @@ class IDataPrepImporterInterface
 public:
 	virtual ~IDataPrepImporterInterface() = default;
 
-	virtual UDatasmithImportOptions* CreateOptions(UObject* Outer) = 0;
-	virtual bool Initialize(const FString& GUID, UWorld* ImportWorld, UPackage* ParentPackage, UDatasmithScene* DatasmithScene, UDatasmithImportOptions* Options) = 0;
+	/**
+	 * @param GUID				The GUID to use as a seed when generating unique ids.
+	 * @param ImportWorld		The destination world that we will spawn the actors in.
+	 * @param DatasmithScene	The DatasmithScene that we will apply the data prep pipeline on.
+	 */
+	virtual bool Initialize(const FGuid& Guid, UWorld* ImportWorld, UDatasmithScene* DatasmithScene) = 0;
 	virtual bool BuildWorld(TArray<TWeakObjectPtr<UObject>>& OutAssets) = 0;
 	virtual bool SetFinalWorld(UWorld* FinalWorld) = 0;
 	virtual bool FinalizeAssets(const TArray<TWeakObjectPtr<UObject>>& Assets) = 0;
