- Fixed Datasmith Area lights template issue because we were running the construction script during our components finalisation step which was invalidating our components.
- Fixed PostProcessVolumes templates not resetting because the object was the component and not the actor.

diff --git a/Engine/Plugins/Enterprise/DatasmithContent/Source/DatasmithContent/Private/ObjectTemplates/DatasmithAreaLightActorTemplate.cpp b/Engine/Plugins/Enterprise/DatasmithContent/Source/DatasmithContent/Private/ObjectTemplates/DatasmithAreaLightActorTemplate.cpp
--- a/Engine/Plugins/Enterprise/DatasmithContent/Source/DatasmithContent/Private/ObjectTemplates/DatasmithAreaLightActorTemplate.cpp
+++ b/Engine/Plugins/Enterprise/DatasmithContent/Source/DatasmithContent/Private/ObjectTemplates/DatasmithAreaLightActorTemplate.cpp
@@ -38,8 +38,6 @@ void UDatasmithAreaLightActorTemplate::Apply( UObject* Destination, bool bForce
 	DATASMITHOBJECTTEMPLATE_CONDITIONALSET( AttenuationRadius, AreaLightActor, PreviousTemplate );
 
 	FDatasmithObjectTemplateUtils::SetObjectTemplate( AreaLightActor->GetRootComponent(), this );
-	AreaLightActor->RerunConstructionScripts();
-
 #endif // #if WITH_EDITORONLY_DATA
 }
 
diff --git a/Engine/Plugins/Enterprise/DatasmithContent/Source/DatasmithContent/Private/ObjectTemplates/DatasmithPostProcessVolumeTemplate.cpp b/Engine/Plugins/Enterprise/DatasmithContent/Source/DatasmithContent/Private/ObjectTemplates/DatasmithPostProcessVolumeTemplate.cpp
--- a/Engine/Plugins/Enterprise/DatasmithContent/Source/DatasmithContent/Private/ObjectTemplates/DatasmithPostProcessVolumeTemplate.cpp
+++ b/Engine/Plugins/Enterprise/DatasmithContent/Source/DatasmithContent/Private/ObjectTemplates/DatasmithPostProcessVolumeTemplate.cpp
@@ -2,13 +2,15 @@
 
 #include "ObjectTemplates/DatasmithPostProcessVolumeTemplate.h"
 
+#include "ObjectTemplates/DatasmithActorTemplate.h"
+
 #include "Engine/PostProcessVolume.h"
 
 void UDatasmithPostProcessVolumeTemplate::Apply( UObject* Destination, bool bForce )
 {
 #if WITH_EDITORONLY_DATA
-	APostProcessVolume* PostProcessVolume = Cast< APostProcessVolume >( Destination );
 
+	APostProcessVolume* PostProcessVolume = UDatasmithActorTemplate::GetActor< APostProcessVolume >( Destination );
 	if ( !PostProcessVolume )
 	{
 		return;
@@ -28,8 +30,7 @@ void UDatasmithPostProcessVolumeTemplate::Apply( UObject* Destination, bool bFor
 void UDatasmithPostProcessVolumeTemplate::Load( const UObject* Source )
 {
 #if WITH_EDITORONLY_DATA
-	const APostProcessVolume* PostProcessVolume = Cast< APostProcessVolume >( Source );
-
+	const APostProcessVolume* PostProcessVolume = UDatasmithActorTemplate::GetActor< const APostProcessVolume >( Source );
 	if ( !PostProcessVolume )
 	{
 		return;
diff --git a/Engine/Plugins/Enterprise/DatasmithContent/Source/DatasmithContent/Public/ObjectTemplates/DatasmithActorTemplate.h b/Engine/Plugins/Enterprise/DatasmithContent/Source/DatasmithContent/Public/ObjectTemplates/DatasmithActorTemplate.h
--- a/Engine/Plugins/Enterprise/DatasmithContent/Source/DatasmithContent/Public/ObjectTemplates/DatasmithActorTemplate.h
+++ b/Engine/Plugins/Enterprise/DatasmithContent/Source/DatasmithContent/Public/ObjectTemplates/DatasmithActorTemplate.h
@@ -26,4 +26,19 @@ public:
 	virtual void Apply( UObject* Destination, bool bForce = false ) override;
 	virtual void Load( const UObject* Source ) override;
 	virtual bool Equals( const UDatasmithObjectTemplate* Other ) const override;
+
+	/** Helper function to get the typed Actor from either a component or an actor */
+	template< typename T >
+	static T* GetActor( UObject* Object )
+	{
+		const UActorComponent* ActorComponent = Cast< UActorComponent >( Object );
+		return Cast< T >( ActorComponent ? ActorComponent->GetOwner() : Object );
+	}
+
+	template< typename T >
+	static const T* GetActor( const UObject* Object )
+	{
+		const UActorComponent* ActorComponent = Cast< UActorComponent >( Object );
+		return Cast< T >( ActorComponent ? static_cast< const AActor* >( ActorComponent->GetOwner() ) : Object );
+	}
 };
\ No newline at end of file
