Fixed ternary operator type deduction issue when compiling on linux.

diff --git a/Engine/Plugins/Enterprise/DatasmithContent/Source/DatasmithContent/Public/ObjectTemplates/DatasmithActorTemplate.h b/Engine/Plugins/Enterprise/DatasmithContent/Source/DatasmithContent/Public/ObjectTemplates/DatasmithActorTemplate.h
--- a/Engine/Plugins/Enterprise/DatasmithContent/Source/DatasmithContent/Public/ObjectTemplates/DatasmithActorTemplate.h
+++ b/Engine/Plugins/Enterprise/DatasmithContent/Source/DatasmithContent/Public/ObjectTemplates/DatasmithActorTemplate.h
@@ -3,6 +3,7 @@
 #pragma once
 
 #include "Engine/EngineTypes.h"
+#include "GameFramework/Actor.h"
 #include "ObjectTemplates/DatasmithObjectTemplate.h"
 
 #include "DatasmithActorTemplate.generated.h"
@@ -31,14 +32,24 @@ public:
 	template< typename T >
 	static T* GetActor( UObject* Object )
 	{
-		const UActorComponent* ActorComponent = Cast< UActorComponent >( Object );
-		return Cast< T >( ActorComponent ? ActorComponent->GetOwner() : Object );
+		return const_cast< T* >( GetActor< T >( static_cast< const UObject* >( Object ) ) );
 	}
 
 	template< typename T >
 	static const T* GetActor( const UObject* Object )
 	{
 		const UActorComponent* ActorComponent = Cast< UActorComponent >( Object );
-		return Cast< T >( ActorComponent ? static_cast< const AActor* >( ActorComponent->GetOwner() ) : Object );
+		const UObject* Actor;
+		
+		if ( ActorComponent )
+		{
+			Actor = static_cast< const UObject* >( ActorComponent->GetOwner() );
+		}
+		else
+		{
+			Actor = Object;
+		}
+
+		return Cast< T >( Actor );
 	}
 };
\ No newline at end of file
