Hot Fix.
Added a const bool as a flag in the datasmith template base class to know if the template is for an actor.
And now, the actors templates are properly migrated.

Note: There might still be some templates that aren't properly flagged.

diff --git a/Engine/Plugins/Enterprise/DatasmithContent/Source/DatasmithContent/Private/ObjectTemplates/DatasmithAreaLightActorTemplate.cpp b/Engine/Plugins/Enterprise/DatasmithContent/Source/DatasmithContent/Private/ObjectTemplates/DatasmithAreaLightActorTemplate.cpp
--- a/Engine/Plugins/Enterprise/DatasmithContent/Source/DatasmithContent/Private/ObjectTemplates/DatasmithAreaLightActorTemplate.cpp
+++ b/Engine/Plugins/Enterprise/DatasmithContent/Source/DatasmithContent/Private/ObjectTemplates/DatasmithAreaLightActorTemplate.cpp
@@ -5,6 +5,7 @@
 #include "Engine/TextureLightProfile.h"
 
 UDatasmithAreaLightActorTemplate::UDatasmithAreaLightActorTemplate()
+	: UDatasmithObjectTemplate(true)
 {
 	Load( ADatasmithAreaLightActor::StaticClass()->GetDefaultObject() );
 }
diff --git a/Engine/Plugins/Enterprise/DatasmithContent/Source/DatasmithContent/Public/ObjectTemplates/DatasmithActorTemplate.h b/Engine/Plugins/Enterprise/DatasmithContent/Source/DatasmithContent/Public/ObjectTemplates/DatasmithActorTemplate.h
--- a/Engine/Plugins/Enterprise/DatasmithContent/Source/DatasmithContent/Public/ObjectTemplates/DatasmithActorTemplate.h
+++ b/Engine/Plugins/Enterprise/DatasmithContent/Source/DatasmithContent/Public/ObjectTemplates/DatasmithActorTemplate.h
@@ -16,6 +16,10 @@ class DATASMITHCONTENT_API UDatasmithActorTemplate : public UDatasmithObjectTemp
 
 public:
 
+	UDatasmithActorTemplate()
+		: UDatasmithObjectTemplate(true)
+	{}
+	  
 	/** Layers this actor belongs to. (see AActor::Layers) */
 	UPROPERTY()
 	TSet<FName> Layers;
diff --git a/Engine/Plugins/Enterprise/DatasmithContent/Source/DatasmithContent/Public/ObjectTemplates/DatasmithCineCameraActorTemplate.h b/Engine/Plugins/Enterprise/DatasmithContent/Source/DatasmithContent/Public/ObjectTemplates/DatasmithCineCameraActorTemplate.h
--- a/Engine/Plugins/Enterprise/DatasmithContent/Source/DatasmithContent/Public/ObjectTemplates/DatasmithCineCameraActorTemplate.h
+++ b/Engine/Plugins/Enterprise/DatasmithContent/Source/DatasmithContent/Public/ObjectTemplates/DatasmithCineCameraActorTemplate.h
@@ -34,6 +34,10 @@ class DATASMITHCONTENT_API UDatasmithCineCameraActorTemplate : public UDatasmith
 	GENERATED_BODY()
 
 public:
+	UDatasmithCineCameraActorTemplate()
+		: UDatasmithObjectTemplate(true)
+	{}
+
 	UPROPERTY()
 	FDatasmithCameraLookatTrackingSettingsTemplate LookatTrackingSettings;
 
diff --git a/Engine/Plugins/Enterprise/DatasmithContent/Source/DatasmithContent/Public/ObjectTemplates/DatasmithLandscapeTemplate.h b/Engine/Plugins/Enterprise/DatasmithContent/Source/DatasmithContent/Public/ObjectTemplates/DatasmithLandscapeTemplate.h
--- a/Engine/Plugins/Enterprise/DatasmithContent/Source/DatasmithContent/Public/ObjectTemplates/DatasmithLandscapeTemplate.h
+++ b/Engine/Plugins/Enterprise/DatasmithContent/Source/DatasmithContent/Public/ObjectTemplates/DatasmithLandscapeTemplate.h
@@ -14,6 +14,10 @@ class DATASMITHCONTENT_API UDatasmithLandscapeTemplate : public UDatasmithObject
 	GENERATED_BODY()
 
 public:
+	UDatasmithLandscapeTemplate()
+		: UDatasmithObjectTemplate(true)
+	{}
+
 	UPROPERTY()
 	UMaterialInterface* LandscapeMaterial;
 
diff --git a/Engine/Plugins/Enterprise/DatasmithContent/Source/DatasmithContent/Public/ObjectTemplates/DatasmithObjectTemplate.h b/Engine/Plugins/Enterprise/DatasmithContent/Source/DatasmithContent/Public/ObjectTemplates/DatasmithObjectTemplate.h
--- a/Engine/Plugins/Enterprise/DatasmithContent/Source/DatasmithContent/Public/ObjectTemplates/DatasmithObjectTemplate.h
+++ b/Engine/Plugins/Enterprise/DatasmithContent/Source/DatasmithContent/Public/ObjectTemplates/DatasmithObjectTemplate.h
@@ -12,6 +12,18 @@ class DATASMITHCONTENT_API UDatasmithObjectTemplate : public UObject
 	GENERATED_BODY()
 
 public:
+
+	UDatasmithObjectTemplate()
+		: UObject()
+		, bIsActorTemplate(false)
+	{}
+
+
+	UDatasmithObjectTemplate(bool bInIsActorTemplate)
+		: UObject()
+		, bIsActorTemplate(bInIsActorTemplate)
+	{}
+
 	/**
 	 * Applies the object template to a Destination object
 	 *
@@ -29,6 +41,9 @@ public:
 	 * Returns if this template equals another template of the same type.
 	 */
 	virtual bool Equals( const UDatasmithObjectTemplate* Other ) const { return false; }
+
+	// Is this template for an actor
+	const bool bIsActorTemplate = false;
 };
 
 // Sets Destination->MemberName with the value of MemberName only if PreviousTemplate is null or has the same value for MemberName as the Destination.
diff --git a/Engine/Plugins/Enterprise/DatasmithContent/Source/DatasmithContent/Public/ObjectTemplates/DatasmithPostProcessVolumeTemplate.h b/Engine/Plugins/Enterprise/DatasmithContent/Source/DatasmithContent/Public/ObjectTemplates/DatasmithPostProcessVolumeTemplate.h
--- a/Engine/Plugins/Enterprise/DatasmithContent/Source/DatasmithContent/Public/ObjectTemplates/DatasmithPostProcessVolumeTemplate.h
+++ b/Engine/Plugins/Enterprise/DatasmithContent/Source/DatasmithContent/Public/ObjectTemplates/DatasmithPostProcessVolumeTemplate.h
@@ -16,6 +16,9 @@ class DATASMITHCONTENT_API UDatasmithPostProcessVolumeTemplate : public UDatasmi
 	GENERATED_BODY()
 
 public:
+	UDatasmithPostProcessVolumeTemplate()
+		: UDatasmithObjectTemplate(true)
+	{}
 
 	UPROPERTY()
 	FDatasmithPostProcessSettingsTemplate Settings;
