Fixed regression with landscape material rendering black when imported through Datasmith. Must call post edit on the material property after the actor has been finalized.

diff --git a/Engine/Plugins/Enterprise/DatasmithContent/Source/DatasmithContent/Private/ObjectTemplates/DatasmithLandscapeTemplate.cpp b/Engine/Plugins/Enterprise/DatasmithContent/Source/DatasmithContent/Private/ObjectTemplates/DatasmithLandscapeTemplate.cpp
--- a/Engine/Plugins/Enterprise/DatasmithContent/Source/DatasmithContent/Private/ObjectTemplates/DatasmithLandscapeTemplate.cpp
+++ b/Engine/Plugins/Enterprise/DatasmithContent/Source/DatasmithContent/Private/ObjectTemplates/DatasmithLandscapeTemplate.cpp
@@ -3,11 +3,12 @@
 #include "ObjectTemplates/DatasmithLandscapeTemplate.h"
 
 #include "Landscape.h"
+#include "ObjectTemplates/DatasmithActorTemplate.h"
 
 void UDatasmithLandscapeTemplate::Apply( UObject* Destination, bool bForce )
 {
 #if WITH_EDITORONLY_DATA
-	ALandscape* Landscape = Cast< ALandscape >( Destination );
+	ALandscape* Landscape = UDatasmithActorTemplate::GetActor< ALandscape >( Destination );
 
 	if( !Landscape )
 	{
@@ -26,7 +27,7 @@ void UDatasmithLandscapeTemplate::Apply( UObject* Destination, bool bForce )
 void UDatasmithLandscapeTemplate::Load( const UObject* Source )
 {
 #if WITH_EDITORONLY_DATA
-	const ALandscape* Landscape = Cast< ALandscape >( Source );
+	const ALandscape* Landscape = UDatasmithActorTemplate::GetActor< ALandscape >( Source );
 
 	if( !Landscape )
 	{
