Retessellate action from StaticMeshEditor: auto registered via ct extention

diff --git a/Engine/Plugins/Enterprise/DatasmithContent/Source/DatasmithContent/Public/DatasmithAssetImportData.h b/Engine/Plugins/Enterprise/DatasmithContent/Source/DatasmithContent/Public/DatasmithAssetImportData.h
--- a/Engine/Plugins/Enterprise/DatasmithContent/Source/DatasmithContent/Public/DatasmithAssetImportData.h
+++ b/Engine/Plugins/Enterprise/DatasmithContent/Source/DatasmithContent/Public/DatasmithAssetImportData.h
@@ -18,7 +18,7 @@ public:
 	UPROPERTY(BlueprintReadWrite, Category = Asset, meta = (ShowOnlyInnerProperties))
 	FDatasmithAssetImportOptions AssetImportOptions;
 
-	UPROPERTY(BlueprintReadWrite, Category = Asset, meta = (ShowOnlyInnerProperties))
+	UPROPERTY(EditAnywhere, Instanced, Category = Asset, meta = (ShowOnlyInnerProperties))
 	TArray<class UDatasmithAdditionalData*> AdditionalData;
 #endif		// WITH_EDITORONLY_DATA
 };
