Fix cut/paste loosing RelatedActors leading to duplicate objects when reimporting a datasmith scene
Fix missing Datasmith_UniqueId on actors for levels saved with UE < 4.22
Fix logged datasmith import time to include file loading like in 4.22
Improve performance when deleting actors when a datasmith scene is present
   -162 s to 44 s for deleting 60 000 actors

From dev-enterprise[REDACTED_HANDLE],7623577,7627372,7638300,7639854

diff --git a/Engine/Plugins/Enterprise/DatasmithContent/Source/DatasmithContent/DatasmithContent.Build.cs b/Engine/Plugins/Enterprise/DatasmithContent/Source/DatasmithContent/DatasmithContent.Build.cs
--- a/Engine/Plugins/Enterprise/DatasmithContent/Source/DatasmithContent/DatasmithContent.Build.cs
+++ b/Engine/Plugins/Enterprise/DatasmithContent/Source/DatasmithContent/DatasmithContent.Build.cs
@@ -14,8 +14,8 @@ namespace UnrealBuildTool.Rules
 					"Core",
 					"CoreUObject",
 					"Engine",
-                    "RenderCore",
-                }
+					"RenderCore",
+				}
 			);
 
 			PrivateDependencyModuleNames.AddRange(
@@ -24,9 +24,19 @@ namespace UnrealBuildTool.Rules
 					"Landscape",
 					"LevelSequence",
 					"MeshDescription",
-                    "Projects",
-                }
+					"Projects"
+				}
 			);
+
+			if (Target.Type == TargetType.Editor)
+			{
+				PrivateDependencyModuleNames.AddRange(
+					new string[]
+					{
+						"UnrealEd"
+					}
+				);
+			}
 		}
 	}
 }
\ No newline at end of file
diff --git a/Engine/Plugins/Enterprise/DatasmithContent/Source/DatasmithContent/Private/DatasmithAssetUserData.cpp b/Engine/Plugins/Enterprise/DatasmithContent/Source/DatasmithContent/Private/DatasmithAssetUserData.cpp
--- a/Engine/Plugins/Enterprise/DatasmithContent/Source/DatasmithContent/Private/DatasmithAssetUserData.cpp
+++ b/Engine/Plugins/Enterprise/DatasmithContent/Source/DatasmithContent/Private/DatasmithAssetUserData.cpp
@@ -4,6 +4,9 @@
 
 #include "DatasmithContentModule.h"
 
+#include "GameFramework/Actor.h"
+#include "Interfaces/Interface_AssetUserData.h"
+
 #if WITH_EDITORONLY_DATA
 
 bool UDatasmithAssetUserData::IsPostLoadThreadSafe() const
@@ -32,3 +35,60 @@ void UDatasmithAssetUserData::PostLoad()
 
 #endif // WITH_EDITORONLY_DATA
 
+const TCHAR* UDatasmithAssetUserData::UniqueIdMetaDataKey = TEXT("Datasmith_UniqueId");
+
+UDatasmithAssetUserData* UDatasmithAssetUserData::GetDatasmithUserData(UObject* Object)
+{
+	if (AActor* Actor = Cast<AActor>(Object))
+	{
+		Object = Actor->GetRootComponent();
+	}
+
+	if (IInterface_AssetUserData* AssetUserData = Cast<IInterface_AssetUserData>(Object))
+	{
+		return Cast<UDatasmithAssetUserData>(AssetUserData->GetAssetUserDataOfClass(UDatasmithAssetUserData::StaticClass()));
+	}
+
+	return nullptr;
+}
+
+FString UDatasmithAssetUserData::GetDatasmithUserDataValueForKey(UObject* Object, FName Key)
+{
+	if (Object)
+	{
+		if (UDatasmithAssetUserData* AssetUserData = GetDatasmithUserData(Object))
+		{
+			FString* Value = AssetUserData->MetaData.Find(Key);
+			return Value ? *Value : FString();
+		}
+	}
+	return FString();
+}
+
+bool UDatasmithAssetUserData::SetDatasmithUserDataValueForKey(UObject* Object, FName Key, const FString & Value)
+{
+	// For AActor, the interface is actually implemented by the ActorComponent
+	if (AActor* Actor = Cast<AActor>(Object))
+	{
+		Object = Actor->GetRootComponent();
+	}
+
+	if (IInterface_AssetUserData* AssetUserData = Cast<IInterface_AssetUserData>(Object))
+	{
+		UDatasmithAssetUserData* DatasmithUserData = AssetUserData->GetAssetUserData< UDatasmithAssetUserData >();
+
+		if (!DatasmithUserData)
+		{
+			DatasmithUserData = NewObject<UDatasmithAssetUserData>(Object, NAME_None, RF_Public | RF_Transactional);
+			AssetUserData->AddAssetUserData(DatasmithUserData);
+		}
+
+		// Add Datasmith meta data
+		DatasmithUserData->MetaData.Add(Key, Value);
+		DatasmithUserData->MetaData.KeySort(FNameLexicalLess());
+
+		return true;
+	}
+
+	return false;
+}
\ No newline at end of file
diff --git a/Engine/Plugins/Enterprise/DatasmithContent/Source/DatasmithContent/Private/DatasmithContentBlueprintLibrary.cpp b/Engine/Plugins/Enterprise/DatasmithContent/Source/DatasmithContent/Private/DatasmithContentBlueprintLibrary.cpp
--- a/Engine/Plugins/Enterprise/DatasmithContent/Source/DatasmithContent/Private/DatasmithContentBlueprintLibrary.cpp
+++ b/Engine/Plugins/Enterprise/DatasmithContent/Source/DatasmithContent/Private/DatasmithContentBlueprintLibrary.cpp
@@ -10,29 +10,12 @@
 
 UDatasmithAssetUserData* UDatasmithContentBlueprintLibrary::GetDatasmithUserData(UObject* Object)
 {
-	if (AActor* Actor = Cast<AActor>(Object))
-	{
-		Object = Actor->GetRootComponent();
-	}
-
-	if (IInterface_AssetUserData* AssetUserData = Cast<IInterface_AssetUserData>(Object))
-	{
-		return Cast<UDatasmithAssetUserData>(AssetUserData->GetAssetUserDataOfClass(UDatasmithAssetUserData::StaticClass()));
-	}
-	return nullptr;
+	return UDatasmithAssetUserData::GetDatasmithUserData(Object);
 }
 
 FString UDatasmithContentBlueprintLibrary::GetDatasmithUserDataValueForKey(UObject* Object, FName Key)
 {
-	if (Object)
-	{
-		if (UDatasmithAssetUserData* AssetUserData = GetDatasmithUserData(Object))
-		{
-			FString* Value = AssetUserData->MetaData.Find(Key);
-			return Value ? *Value : FString();
-		}
-	}
-	return FString();
+	return UDatasmithAssetUserData::GetDatasmithUserDataValueForKey(Object, Key);
 }
 
 void UDatasmithContentBlueprintLibrary::GetDatasmithUserDataKeysAndValuesForValue(UObject* Object, const FString& StringToMatch, TArray<FName>& OutKeys, TArray<FString>& OutValues)
diff --git a/Engine/Plugins/Enterprise/DatasmithContent/Source/DatasmithContent/Private/DatasmithSceneActor.cpp b/Engine/Plugins/Enterprise/DatasmithContent/Source/DatasmithContent/Private/DatasmithSceneActor.cpp
--- a/Engine/Plugins/Enterprise/DatasmithContent/Source/DatasmithContent/Private/DatasmithSceneActor.cpp
+++ b/Engine/Plugins/Enterprise/DatasmithContent/Source/DatasmithContent/Private/DatasmithSceneActor.cpp
@@ -1,43 +1,131 @@
 // Copyright 1998-2019 Epic Games, Inc. All Rights Reserved.
 
 #include "DatasmithSceneActor.h"
+#include "DatasmithAssetUserData.h"
+#include "DatasmithContentModule.h"
 
 #include "Engine/Engine.h"
 
+#if WITH_EDITOR
+#include "Editor.h"
+#endif
+
 ADatasmithSceneActor::ADatasmithSceneActor()
 {
 #if WITH_EDITOR
+	FEditorDelegates::MapChange.AddUObject(this, &ADatasmithSceneActor::OnMapChange);
+
 	if ( GEngine )
 	{
 		OnActorDeletedDelegateHandle = GEngine->OnLevelActorDeleted().AddUObject( this, &ADatasmithSceneActor::OnActorDeleted );
 	}
+
+	if ( GEditor )
+	{
+		GEditor->GetEditorSubsystem<UImportSubsystem>()->OnAssetPostImport.AddUObject(this, &ADatasmithSceneActor::OnAssetPostImport);
+	}
 #endif // WITH_EDITOR
 }
 
 ADatasmithSceneActor::~ADatasmithSceneActor()
 {
 #if WITH_EDITOR
+	FEditorDelegates::MapChange.RemoveAll(this);
+
 	if ( GEngine )
 	{
 		GEngine->OnLevelActorDeleted().Remove( OnActorDeletedDelegateHandle );
 	}
+
+	if ( GEditor )
+	{
+		GEditor->GetEditorSubsystem<UImportSubsystem>()->OnAssetPostImport.RemoveAll(this);
+	}
 #endif // WITH_EDITOR
 }
 
-
 #if WITH_EDITOR
+
+void ADatasmithSceneActor::PostLoad()
+{
+	EnsureDatasmithIdsForRelatedActors();
+
+	Super::PostLoad();
+}
+
+// Since we now rely on all datasmith related object to have a UniqueId, this will take care of fixing all related actors
+void ADatasmithSceneActor::EnsureDatasmithIdsForRelatedActors()
+{
+	for (TPair< FName, TSoftObjectPtr< AActor > >& ActorPair : RelatedActors)
+	{
+		if (AActor* Actor = ActorPair.Value.Get())
+		{
+			FString DatasmithUniqueId = UDatasmithAssetUserData::GetDatasmithUserDataValueForKey(Actor, UDatasmithAssetUserData::UniqueIdMetaDataKey);
+			if (DatasmithUniqueId.IsEmpty())
+			{
+				if (!UDatasmithAssetUserData::SetDatasmithUserDataValueForKey(Actor, UDatasmithAssetUserData::UniqueIdMetaDataKey, ActorPair.Key.ToString()))
+				{
+					UE_LOG(LogDatasmithContent, Warning, TEXT("Actor %s is referenced by datasmith scene but doesn't have a datasmith uniqueId."), *Actor->GetName());
+				}
+			}
+		}
+	}
+}
+
+void ADatasmithSceneActor::OnMapChange(uint32 MapEventFlags)
+{
+	// Whenever a sub-level is loaded, we need to apply the fix
+	if (MapEventFlags == MapChangeEventFlags::NewMap)
+	{
+		EnsureDatasmithIdsForRelatedActors();
+	}
+}
+
+// Reset the related actor on the scene when an actor is deleted to avoid re-matching it with a manually created object that might have the same name.
 void ADatasmithSceneActor::OnActorDeleted(AActor* ActorDeleted)
 {
-	for ( TPair< FName, TSoftObjectPtr< AActor > >& Pair : RelatedActors )
+	FString DatasmithUniqueId = UDatasmithAssetUserData::GetDatasmithUserDataValueForKey(ActorDeleted, UDatasmithAssetUserData::UniqueIdMetaDataKey);
+	if (!DatasmithUniqueId.IsEmpty())
 	{
-		AActor* RelatedActor = Pair.Value.Get();
-		if ( RelatedActor == ActorDeleted )
+		TSoftObjectPtr< AActor >* RelatedActorPtr = RelatedActors.Find(FName(*DatasmithUniqueId));
+		if (RelatedActorPtr)
 		{
-			// this will add this actor to the transaction if there is one currently recording
-			Modify();
+			AActor* RelatedActor = RelatedActorPtr->Get();
+			if (RelatedActor == ActorDeleted)
+			{
+				// this will add this actor to the transaction if there is one currently recording
+				Modify();
 
-			Pair.Value.Reset();
+				RelatedActorPtr->Reset();
+			}
 		}
 	}
 }
+
+// Reattach related actors on cut/paste by matching their Datasmith UniqueId
+void ADatasmithSceneActor::OnAssetPostImport(UFactory* InFactory, UObject* ActorAdded)
+{
+	for (TObjectIterator<AActor> It; It; ++It)
+	{
+		AActor* Actor = *It;
+
+		FString DatasmithUniqueId = UDatasmithAssetUserData::GetDatasmithUserDataValueForKey(Actor, UDatasmithAssetUserData::UniqueIdMetaDataKey);
+		if (!DatasmithUniqueId.IsEmpty())
+		{
+			TSoftObjectPtr< AActor >* RelatedActorPtr = RelatedActors.Find(FName(*DatasmithUniqueId));
+			if (RelatedActorPtr)
+			{
+				AActor* RelatedActor = RelatedActorPtr->Get();
+				if (RelatedActor == nullptr)
+				{
+					// this will add this actor to the transaction if there is one currently recording
+					Modify();
+
+					*RelatedActorPtr = Actor;
+				}
+			}
+		}
+	}
+}
+
 #endif // WITH_EDITOR
diff --git a/Engine/Plugins/Enterprise/DatasmithContent/Source/DatasmithContent/Public/DatasmithAssetUserData.h b/Engine/Plugins/Enterprise/DatasmithContent/Source/DatasmithContent/Public/DatasmithAssetUserData.h
--- a/Engine/Plugins/Enterprise/DatasmithContent/Source/DatasmithContent/Public/DatasmithAssetUserData.h
+++ b/Engine/Plugins/Enterprise/DatasmithContent/Source/DatasmithContent/Public/DatasmithAssetUserData.h
@@ -28,4 +28,11 @@ public:
 	virtual bool IsPostLoadThreadSafe() const override;
 	virtual void PostLoad() override;
 #endif
+
+	static FString GetDatasmithUserDataValueForKey(UObject* Object, FName Key);
+	static UDatasmithAssetUserData* GetDatasmithUserData(UObject* Object);
+	static bool SetDatasmithUserDataValueForKey(UObject* Object, FName Key, const FString & Value);
+
+	// Meta data keys for Datasmith objects
+	static const TCHAR* UniqueIdMetaDataKey;
 };
diff --git a/Engine/Plugins/Enterprise/DatasmithContent/Source/DatasmithContent/Public/DatasmithSceneActor.h b/Engine/Plugins/Enterprise/DatasmithContent/Source/DatasmithContent/Public/DatasmithSceneActor.h
--- a/Engine/Plugins/Enterprise/DatasmithContent/Source/DatasmithContent/Public/DatasmithSceneActor.h
+++ b/Engine/Plugins/Enterprise/DatasmithContent/Source/DatasmithContent/Public/DatasmithSceneActor.h
@@ -6,6 +6,8 @@
 
 #include "DatasmithSceneActor.generated.h"
 
+class UFactory;
+
 UCLASS()
 class DATASMITHCONTENT_API ADatasmithSceneActor : public AActor
 {
@@ -26,9 +28,21 @@ public:
 
 private:
 #if WITH_EDITOR
+	// For backward compatibility with UE <= 4.21
+	void EnsureDatasmithIdsForRelatedActors();
+
+	// Backward compatibility fix will happen on PostLoad.
+	void PostLoad() override;
+
+	// We need this to monitor when a sub-level is loaded
+	void OnMapChange(uint32 MapChangeFlags);
+
 	// Cleans the invalid Soft Object Ptr
 	void OnActorDeleted(AActor* ActorDeleted);
 
+	// Reattach newly imported actors if they've been moved / cut and pasted
+	void OnAssetPostImport(UFactory* InFactory, UObject* ActorAdded);
+
 	FDelegateHandle OnActorDeletedDelegateHandle;
 #endif // WITH_EDITOR
 };
