Fixed logic gap in UDatasmithStaticMeshComponentTemplate::UpdateObject when number of override materials are different
Found thanks to [REDACTED_PERSON]!

#rnx

diff --git a/Engine/Plugins/Enterprise/DatasmithContent/Source/DatasmithContent/Private/ObjectTemplates/DatasmithStaticMeshComponentTemplate.cpp b/Engine/Plugins/Enterprise/DatasmithContent/Source/DatasmithContent/Private/ObjectTemplates/DatasmithStaticMeshComponentTemplate.cpp
--- a/Engine/Plugins/Enterprise/DatasmithContent/Source/DatasmithContent/Private/ObjectTemplates/DatasmithStaticMeshComponentTemplate.cpp
+++ b/Engine/Plugins/Enterprise/DatasmithContent/Source/DatasmithContent/Private/ObjectTemplates/DatasmithStaticMeshComponentTemplate.cpp
@@ -45,7 +45,8 @@ UObject* UDatasmithStaticMeshComponentTemplate::UpdateObject( UObject* Destinati
 	{
 		for ( int32 MaterialIndexToRemove = PreviousStaticMeshTemplate->OverrideMaterials.Num() - 1; MaterialIndexToRemove >= OverrideMaterials.Num(); --MaterialIndexToRemove )
 		{
-			if ( StaticMeshComponent->OverrideMaterials.IsValidIndex( MaterialIndexToRemove ) )
+			if ( StaticMeshComponent->OverrideMaterials.IsValidIndex( MaterialIndexToRemove ) &&
+				StaticMeshComponent->OverrideMaterials[MaterialIndexToRemove] == PreviousStaticMeshTemplate->OverrideMaterials[MaterialIndexToRemove] )
 			{
 				StaticMeshComponent->OverrideMaterials.RemoveAt( MaterialIndexToRemove );
 			}
