Moving HDRIBackdrop to its own plugin.

diff --git a/Engine/Plugins/Enterprise/DatasmithContent/Source/DatasmithContentEditor/DatasmithContentEditor.Build.cs b/Engine/Plugins/Enterprise/DatasmithContent/Source/DatasmithContentEditor/DatasmithContentEditor.Build.cs
--- a/Engine/Plugins/Enterprise/DatasmithContent/Source/DatasmithContentEditor/DatasmithContentEditor.Build.cs
+++ b/Engine/Plugins/Enterprise/DatasmithContent/Source/DatasmithContentEditor/DatasmithContentEditor.Build.cs
@@ -27,7 +27,6 @@ namespace UnrealBuildTool.Rules
 					"PropertyEditor",
 					"SlateCore",
 					"Slate",
-                    "PlacementMode",
                 }
 			);
 		}
diff --git a/Engine/Plugins/Enterprise/DatasmithContent/Source/DatasmithContentEditor/Private/DatasmithContentEditorModule.cpp b/Engine/Plugins/Enterprise/DatasmithContent/Source/DatasmithContentEditor/Private/DatasmithContentEditorModule.cpp
--- a/Engine/Plugins/Enterprise/DatasmithContent/Source/DatasmithContentEditor/Private/DatasmithContentEditorModule.cpp
+++ b/Engine/Plugins/Enterprise/DatasmithContent/Source/DatasmithContentEditor/Private/DatasmithContentEditorModule.cpp
@@ -12,8 +12,6 @@
 #include "Modules/ModuleManager.h"
 #include "PropertyEditorModule.h"
 
-#include "DatasmithPlacement.h"
-
 #define LOCTEXT_NAMESPACE "DatasmithContentEditorModule"
 
 EAssetTypeCategories::Type IDatasmithContentEditorModule::DatasmithAssetCategoryBit;
@@ -39,17 +37,9 @@ public:
 		AssetTools.RegisterAssetTypeActions(DatasmithSceneAssetTypeAction.ToSharedRef());
 		AssetTypeActionsArray.Add(DatasmithSceneAssetTypeAction);
 
-		FCoreDelegates::OnPostEngineInit.AddRaw(this, &FDatasmithContentEditorModule::OnPostEngineInit);
-
 		FDatasmithContentEditorStyle::Initialize();
 	}
 
-	// Register classes in placement mode
-	void OnPostEngineInit()
-	{
-		FDatasmithPlacement::RegisterPlacement();
-	}
-
 	virtual void ShutdownModule() override
 	{
 		// Unregister the details customization
@@ -73,8 +63,6 @@ public:
 
 		// Shutdown style set associated with datasmith content
 		FDatasmithContentEditorStyle::Shutdown();
-
-		FCoreDelegates::OnPostEngineInit.RemoveAll(this);
 	}
 
 	virtual void RegisterSpawnDatasmithSceneActorsHandler( FOnSpawnDatasmithSceneActors InSpawnActorsDelegate ) override
diff --git a/Engine/Plugins/Enterprise/DatasmithContent/Source/DatasmithContentEditor/Private/DatasmithPlacement.cpp b/Engine/Plugins/Enterprise/DatasmithContent/Source/DatasmithContentEditor/Private/DatasmithPlacement.cpp
deleted file mode 100644
diff --git a/Engine/Plugins/Enterprise/DatasmithContent/Source/DatasmithContentEditor/Private/DatasmithPlacement.h b/Engine/Plugins/Enterprise/DatasmithContent/Source/DatasmithContentEditor/Private/DatasmithPlacement.h
deleted file mode 100644
diff --git a/Engine/Plugins/Runtime/HDRIBackdrop/HDRIBackdrop.uplugin b/Engine/Plugins/Runtime/HDRIBackdrop/HDRIBackdrop.uplugin
new file mode 100644
--- /dev/null
+++ b/Engine/Plugins/Runtime/HDRIBackdrop/HDRIBackdrop.uplugin
@@ -0,0 +1,24 @@
+{
+	"FileVersion": 3,
+	"Version": 1,
+	"VersionName": "1.0",
+	"FriendlyName": "HDRIBackdrop",
+	"Description": "",
+	"Category": "Rendering",
+	"CreatedBy": "Epic Games, Inc.",
+	"CreatedByURL": "http://epicgames.com",
+	"DocsURL": "",
+	"MarketplaceURL": "",
+	"SupportURL": "",
+	"EnabledByDefault": false,
+	"CanContainContent": true,
+	"IsBetaVersion": false,
+	"Installed": false,
+	"Modules": [
+		{
+			"Name": "HDRIBackdrop",
+			"Type": "Editor",
+			"LoadingPhase": "PostEngineInit"
+		}
+	]
+}
\ No newline at end of file
diff --git a/Engine/Plugins/Runtime/HDRIBackdrop/Source/HDRIBackdrop/HDRIBackdrop.Build.cs b/Engine/Plugins/Runtime/HDRIBackdrop/Source/HDRIBackdrop/HDRIBackdrop.Build.cs
new file mode 100644
--- /dev/null
+++ b/Engine/Plugins/Runtime/HDRIBackdrop/Source/HDRIBackdrop/HDRIBackdrop.Build.cs
@@ -0,0 +1,56 @@
+// Copyright 1998-2019 Epic Games, Inc. All Rights Reserved.
+
+using UnrealBuildTool;
+
+public class HDRIBackdrop : ModuleRules
+{
+	public HDRIBackdrop(ReadOnlyTargetRules Target) : base(Target)
+	{
+		PCHUsage = ModuleRules.PCHUsageMode.UseExplicitOrSharedPCHs;
+		
+		PublicIncludePaths.AddRange(
+			new string[] {
+				// ... add public include paths required here ...
+			}
+			);
+				
+		
+		PrivateIncludePaths.AddRange(
+			new string[] {
+				// ... add other private include paths required here ...
+			}
+			);
+			
+		
+		PublicDependencyModuleNames.AddRange(
+			new string[]
+			{
+				"Core",
+				// ... add other public dependencies that you statically link with here ...
+			}
+			);
+			
+		
+		PrivateDependencyModuleNames.AddRange(
+			new string[]
+			{
+				"CoreUObject",
+				"Engine",
+				"Projects",
+				"Slate",
+				"SlateCore",
+				"UnrealEd",
+                "PlacementMode",
+				// ... add private dependencies that you statically link with here ...	
+			}
+			);
+		
+		
+		DynamicallyLoadedModuleNames.AddRange(
+			new string[]
+			{
+				// ... add any modules that your module loads dynamically here ...
+			}
+			);
+	}
+}
diff --git a/Engine/Plugins/Runtime/HDRIBackdrop/Source/HDRIBackdrop/Private/HDRIBackdrop.cpp b/Engine/Plugins/Runtime/HDRIBackdrop/Source/HDRIBackdrop/Private/HDRIBackdrop.cpp
new file mode 100644
--- /dev/null
+++ b/Engine/Plugins/Runtime/HDRIBackdrop/Source/HDRIBackdrop/Private/HDRIBackdrop.cpp
@@ -0,0 +1,22 @@
+// Copyright 1998-2019 Epic Games, Inc. All Rights Reserved.
+
+#include "HDRIBackdrop.h"
+#include "HDRIBackdropPlacement.h"
+#include "HDRIBackdropStyle.h"
+
+#define LOCTEXT_NAMESPACE "FHDRIBackdropModule"
+
+void FHDRIBackdropModule::StartupModule()
+{
+	FHDRIBackdropStyle::Initialize();
+	FHDRIBackdropPlacement::RegisterPlacement();
+}
+
+void FHDRIBackdropModule::ShutdownModule()
+{
+	FHDRIBackdropStyle::Shutdown();
+}
+
+#undef LOCTEXT_NAMESPACE
+	
+IMPLEMENT_MODULE(FHDRIBackdropModule, HDRIBackdrop)
\ No newline at end of file
diff --git a/Engine/Plugins/Runtime/HDRIBackdrop/Source/HDRIBackdrop/Private/HDRIBackdropPlacement.cpp b/Engine/Plugins/Runtime/HDRIBackdrop/Source/HDRIBackdrop/Private/HDRIBackdropPlacement.cpp
new file mode 100644
--- /dev/null
+++ b/Engine/Plugins/Runtime/HDRIBackdrop/Source/HDRIBackdrop/Private/HDRIBackdropPlacement.cpp
@@ -0,0 +1,32 @@
+// Copyright 1998-2019 Epic Games, Inc. All Rights Reserved.
+
+#include "HDRIBackdropPlacement.h"
+#include "IPlacementModeModule.h"
+#include "ActorFactories/ActorFactoryBlueprint.h"
+
+#define LOCTEXT_NAMESPACE "FHDRIBackdropModule"
+
+void FHDRIBackdropPlacement::RegisterPlacement()
+{	
+	UBlueprint* HDRIBackdrop = Cast<UBlueprint>(FSoftObjectPath(TEXT("/HDRIBackdrop/Blueprints/HDRIBackdrop.HDRIBackdrop")).TryLoad());
+	if (HDRIBackdrop == nullptr)
+	{
+		return;
+	}
+	
+	IPlacementModeModule& PlacementModeModule = IPlacementModeModule::Get();
+	FPlacementCategoryInfo Info = *PlacementModeModule.GetRegisteredPlacementCategory(FBuiltInPlacementCategories::Lights());
+
+	FPlaceableItem* BPPlacement = new FPlaceableItem(
+		*UActorFactoryBlueprint::StaticClass(),
+		FAssetData(HDRIBackdrop, true),
+		FName("HDRIBackdrop.ModesIcon"),
+		TOptional<FLinearColor>(),
+		TOptional<int32>(),
+		NSLOCTEXT("PlacementMode", "HDRI Backdrop", "HDRI Backdrop")
+	);
+
+	IPlacementModeModule::Get().RegisterPlaceableItem( Info.UniqueHandle, MakeShareable(BPPlacement) );
+}
+
+#undef LOCTEXT_NAMESPACE
\ No newline at end of file
diff --git a/Engine/Plugins/Runtime/HDRIBackdrop/Source/HDRIBackdrop/Public/HDRIBackdrop.h b/Engine/Plugins/Runtime/HDRIBackdrop/Source/HDRIBackdrop/Public/HDRIBackdrop.h
new file mode 100644
--- /dev/null
+++ b/Engine/Plugins/Runtime/HDRIBackdrop/Source/HDRIBackdrop/Public/HDRIBackdrop.h
@@ -0,0 +1,15 @@
+// Copyright 1998-2019 Epic Games, Inc. All Rights Reserved.
+
+#pragma once
+
+#include "CoreMinimal.h"
+#include "Modules/ModuleManager.h"
+
+class FHDRIBackdropModule : public IModuleInterface
+{
+public:
+
+	/** IModuleInterface implementation */
+	virtual void StartupModule() override;
+	virtual void ShutdownModule() override;
+};
diff --git a/Engine/Plugins/Runtime/HDRIBackdrop/Source/HDRIBackdrop/Public/HDRIBackdropPlacement.h b/Engine/Plugins/Runtime/HDRIBackdrop/Source/HDRIBackdrop/Public/HDRIBackdropPlacement.h
new file mode 100644
--- /dev/null
+++ b/Engine/Plugins/Runtime/HDRIBackdrop/Source/HDRIBackdrop/Public/HDRIBackdropPlacement.h
@@ -0,0 +1,11 @@
+// Copyright 1998-2019 Epic Games, Inc. All Rights Reserved.
+
+#pragma once
+
+#include "CoreMinimal.h"
+
+class FHDRIBackdropPlacement
+{
+public:
+	static void RegisterPlacement();
+};
\ No newline at end of file
