Removed HDRIBackgroup icon from DatasmithContentEditorStyle. This was moved in it's owned pluggin (see cl 8229528).

#rnx

diff --git a/Engine/Plugins/Enterprise/DatasmithContent/Source/DatasmithContentEditor/Private/DatasmithContentEditorStyle.cpp b/Engine/Plugins/Enterprise/DatasmithContent/Source/DatasmithContentEditor/Private/DatasmithContentEditorStyle.cpp
--- a/Engine/Plugins/Enterprise/DatasmithContent/Source/DatasmithContentEditor/Private/DatasmithContentEditorStyle.cpp
+++ b/Engine/Plugins/Enterprise/DatasmithContent/Source/DatasmithContentEditor/Private/DatasmithContentEditorStyle.cpp
@@ -30,8 +30,6 @@ void FDatasmithContentEditorStyle::Initialize()
 	const FVector2D Icon20x20(20.0f, 20.0f);
 	const FVector2D Icon40x40(40.0f, 40.0f);
 
-	StyleSet->Set("DatasmithPlacement.HDRIBackdrop", new IMAGE_PLUGIN_BRUSH("Icons/HDRIBackdrop", Icon40x40));
-
 	StyleSet->Set("DatasmithDataprepEditor.Importer", new IMAGE_PLUGIN_BRUSH("Icons/DatasmithImporterIcon40", Icon40x40));
 	StyleSet->Set("DatasmithDataprepEditor.Importer.Small", new IMAGE_PLUGIN_BRUSH("Icons/DatasmithImporterIcon40", Icon20x20));
 	StyleSet->Set("DatasmithDataprepEditor.Importer.Selected", new IMAGE_PLUGIN_BRUSH("Icons/DatasmithImporterIcon40", Icon40x40));
