Datasmith: fix reimport options displayed even in first import

diff --git a/Engine/Plugins/Enterprise/DatasmithContent/Source/DatasmithContent/Public/DatasmithImportOptions.h b/Engine/Plugins/Enterprise/DatasmithContent/Source/DatasmithContent/Public/DatasmithImportOptions.h
--- a/Engine/Plugins/Enterprise/DatasmithContent/Source/DatasmithContent/Public/DatasmithImportOptions.h
+++ b/Engine/Plugins/Enterprise/DatasmithContent/Source/DatasmithContent/Public/DatasmithImportOptions.h
@@ -250,7 +250,7 @@ struct DATASMITHCONTENT_API FDatasmithTessellationOptions
 	 * Stitching technique applied on neighbouring surfaces before tessellation.
 	 * None : No stitching applied. This is the default.
 	 * Sewing : Connects surfaces which physically share a boundary but not topologically within a set of objects.
-	 *          This technique can modify the structure of the model by removing and adding objects.  
+	 *          This technique can modify the structure of the model by removing and adding objects.
 	 * Healing : Connects surfaces which physically share a boundary but not topologically within an object.
 	 * The techniques are using the chord tolerance to determine if two surfaces should be stitched.
 	 */
@@ -337,17 +337,17 @@ public:
 	FDatasmithImportBaseOptions BaseOptions;
 
 	/** Options specific to the reimport process */
-	UPROPERTY(config, EditAnywhere, BlueprintReadWrite, Category = "NotVisible", meta = (ShowOnlyInnerProperties))
+	UPROPERTY(config, EditAnywhere, BlueprintReadWrite, Category = "Reimport", meta = (ShowOnlyInnerProperties))
 	FDatasmithReimportOptions ReimportOptions;
 
 	/** Name of the imported file without its path */
 	UPROPERTY(BlueprintReadWrite, Category = "NotVisible")
 	FString FileName;
-	
+
 	/** Full path of the imported file */
 	UPROPERTY(BlueprintReadWrite, Category = "NotVisible")
 	FString FilePath;
-	
+
 	/** Whether to use or not the same options when loading multiple files. Default false */
 	bool bUseSameOptions;
 
diff --git a/Engine/Plugins/Enterprise/DatasmithImporter/Source/DatasmithImporter/Private/DatasmithImportContext.cpp b/Engine/Plugins/Enterprise/DatasmithImporter/Source/DatasmithImporter/Private/DatasmithImportContext.cpp
--- a/Engine/Plugins/Enterprise/DatasmithImporter/Source/DatasmithImporter/Private/DatasmithImportContext.cpp
+++ b/Engine/Plugins/Enterprise/DatasmithImporter/Source/DatasmithImporter/Private/DatasmithImportContext.cpp
@@ -328,13 +328,17 @@ void FDatasmithImportContext::DisplayMessages()
 
 void FDatasmithImportContext::SetupBaseOptionsVisibility()
 {
-	if ( bIsAReimport )
+	if ( UProperty* ReimportOptionsProperty = FindField< UProperty >( Options->GetClass(), GET_MEMBER_NAME_CHECKED( UDatasmithImportOptions, ReimportOptions ) ) )
 	{
-		UProperty* ReimportOptionsProperty = FindField< UProperty >( Options->GetClass(), GET_MEMBER_NAME_CHECKED( UDatasmithImportOptions, ReimportOptions ) );
-
-		if ( ReimportOptionsProperty )
+		if ( bIsAReimport )
 		{
 			ReimportOptionsProperty->SetMetaData( TEXT("Category"), TEXT("Reimport") );
+			ReimportOptionsProperty->SetMetaData( TEXT("ShowOnlyInnerProperties"), TEXT("1") );
+		}
+		else
+		{
+			ReimportOptionsProperty->SetMetaData( TEXT("Category"), TEXT("NotVisible") );
+			ReimportOptionsProperty->RemoveMetaData( TEXT("ShowOnlyInnerProperties") );
 		}
 	}
 }
@@ -345,6 +349,7 @@ void FDatasmithImportContext::ResetBaseOptionsVisibility()
 
 	if ( ReimportOptionsProperty )
 	{
+		ReimportOptionsProperty->SetMetaData( TEXT("ShowOnlyInnerProperties"), TEXT("1") );
 		ReimportOptionsProperty->SetMetaData( TEXT("Category"), TEXT("NotVisible") );
 	}
 }
