Make datasmith scenes re-execute the Dataprep pipeline #jira UEENT-3018 #rb JeanLuc.Corenthin

diff --git a/Engine/Plugins/Enterprise/DatasmithContent/DatasmithContent.uplugin b/Engine/Plugins/Enterprise/DatasmithContent/DatasmithContent.uplugin
--- a/Engine/Plugins/Enterprise/DatasmithContent/DatasmithContent.uplugin
+++ b/Engine/Plugins/Enterprise/DatasmithContent/DatasmithContent.uplugin
@@ -30,6 +30,14 @@
 		{
 			"Name": "VariantManagerContent",
 			"Enabled": true
+		},
+		{
+			"Name": "DataprepEditor",
+			"Enabled": true
+		},
+		{
+			"Name": "EditorScriptingUtilities",
+			"Enabled": true
 		}
 	]
 }
\ No newline at end of file
diff --git a/Engine/Plugins/Enterprise/DatasmithContent/Source/DatasmithContent/Private/DatasmithScene.cpp b/Engine/Plugins/Enterprise/DatasmithContent/Source/DatasmithContent/Private/DatasmithScene.cpp
--- a/Engine/Plugins/Enterprise/DatasmithContent/Source/DatasmithContent/Private/DatasmithScene.cpp
+++ b/Engine/Plugins/Enterprise/DatasmithContent/Source/DatasmithContent/Private/DatasmithScene.cpp
@@ -3,6 +3,7 @@
 #include "DatasmithScene.h"
 #include "DatasmithAssetImportData.h"
 
+#include "Engine/AssetUserData.h"
 #include "Engine/Blueprint.h"
 #include "Engine/World.h"
 #include "EngineUtils.h"
@@ -21,9 +22,6 @@ enum
 
 UDatasmithScene::UDatasmithScene()
 {
-#if WITH_EDITORONLY_DATA
-	DataprepRecipeBP = nullptr;
-#endif
 #if WITH_EDITOR
 	bPreWorldRenameCallbackRegistered = false;
 #endif
@@ -107,3 +105,58 @@ void UDatasmithScene::Serialize( FArchive& Archive )
 	}
 #endif // #if WITH_EDITORONLY_DATA
 }
+
+void UDatasmithScene::AddAssetUserData( UAssetUserData* InUserData )
+{
+#if WITH_EDITORONLY_DATA
+	if ( InUserData != NULL )
+	{
+		UAssetUserData* ExistingData = GetAssetUserDataOfClass( InUserData->GetClass() );
+		if ( ExistingData != NULL )
+		{
+			AssetUserData.Remove( ExistingData );
+		}
+		AssetUserData.Add( InUserData );
+	}
+#endif // #if WITH_EDITORONLY_DATA
+}
+
+UAssetUserData* UDatasmithScene::GetAssetUserDataOfClass( TSubclassOf<UAssetUserData> InUserDataClass )
+{
+#if WITH_EDITORONLY_DATA
+	for ( int32 DataIdx = 0; DataIdx < AssetUserData.Num(); DataIdx++ )
+	{
+		UAssetUserData* Datum = AssetUserData[DataIdx];
+		if ( Datum != NULL && Datum->IsA( InUserDataClass ) )
+		{
+			return Datum;
+		}
+	}
+#endif // #if WITH_EDITORONLY_DATA
+	return NULL;
+
+}
+
+void UDatasmithScene::RemoveUserDataOfClass( TSubclassOf<UAssetUserData> InUserDataClass )
+{
+#if WITH_EDITORONLY_DATA
+	for ( int32 DataIdx = 0; DataIdx < AssetUserData.Num(); DataIdx++ )
+	{
+		UAssetUserData* Datum = AssetUserData[DataIdx];
+		if ( Datum != NULL && Datum->IsA(InUserDataClass ) )
+		{
+			AssetUserData.RemoveAt( DataIdx );
+			return;
+		}
+	}
+#endif // #if WITH_EDITORONLY_DATA
+}
+
+const TArray<UAssetUserData*>* UDatasmithScene::GetAssetUserDataArray() const
+{
+#if WITH_EDITORONLY_DATA
+	return &AssetUserData;
+#else
+	return NULL;
+#endif // #if WITH_EDITORONLY_DATA
+}
diff --git a/Engine/Plugins/Enterprise/DatasmithContent/Source/DatasmithContent/Public/DatasmithScene.h b/Engine/Plugins/Enterprise/DatasmithContent/Source/DatasmithContent/Public/DatasmithScene.h
--- a/Engine/Plugins/Enterprise/DatasmithContent/Source/DatasmithContent/Public/DatasmithScene.h
+++ b/Engine/Plugins/Enterprise/DatasmithContent/Source/DatasmithContent/Public/DatasmithScene.h
@@ -5,6 +5,7 @@
 #include "DatasmithImportOptions.h"
 #include "Misc/Guid.h"
 #include "Serialization/BulkData.h"
+#include "Interfaces/Interface_AssetUserData.h"
 
 #include "DatasmithScene.generated.h"
 
@@ -17,7 +18,7 @@ class UTexture;
 class UWorld;
 
 UCLASS()
-class DATASMITHCONTENT_API UDatasmithScene : public UObject
+class DATASMITHCONTENT_API UDatasmithScene : public UObject, public IInterface_AssetUserData
 {
 	GENERATED_BODY()
 
@@ -27,10 +28,6 @@ public:
 	virtual ~UDatasmithScene();
 
 #if WITH_EDITORONLY_DATA
-	/** Pointer to data preparation pipeline blueprint used to process input data */
-	UPROPERTY()
-	class UBlueprint* DataprepRecipeBP;
-
 	/** Importing data and options used for this Datasmith scene */
 	UPROPERTY(EditAnywhere, Instanced, Category=ImportSettings)
 	class UDatasmithSceneImportData* AssetImportData;
@@ -63,11 +60,22 @@ public:
 	/** Map of all the level variant sets related to this Datasmith Scene */
 	UPROPERTY(VisibleAnywhere, Category="Datasmith", AdvancedDisplay)
 	TMap< FName, TSoftObjectPtr< ULevelVariantSets > > LevelVariantSets;
+
+	/** Array of user data stored with the asset */
+	UPROPERTY()
+	TArray< UAssetUserData* > AssetUserData;
 #endif // #if WITH_EDITORONLY_DATA
 
 	/** Register the DatasmithScene to the PreWorldRename callback as needed*/
 	void RegisterPreWorldRenameCallback();
 
+	//~ Begin IInterface_AssetUserData Interface
+	virtual void AddAssetUserData( UAssetUserData* InUserData ) override;
+	virtual void RemoveUserDataOfClass( TSubclassOf<UAssetUserData> InUserDataClass ) override;
+	virtual UAssetUserData* GetAssetUserDataOfClass( TSubclassOf<UAssetUserData> InUserDataClass ) override;
+	virtual const TArray<UAssetUserData*>* GetAssetUserDataArray() const override;
+	//~ End IInterface_AssetUserData Interface
+
 #if WITH_EDITOR
 private:
 	/** Called before a world is renamed */
diff --git a/Engine/Plugins/Enterprise/DatasmithContent/Source/DatasmithContentEditor/DatasmithContentEditor.Build.cs b/Engine/Plugins/Enterprise/DatasmithContent/Source/DatasmithContentEditor/DatasmithContentEditor.Build.cs
--- a/Engine/Plugins/Enterprise/DatasmithContent/Source/DatasmithContentEditor/DatasmithContentEditor.Build.cs
+++ b/Engine/Plugins/Enterprise/DatasmithContent/Source/DatasmithContentEditor/DatasmithContentEditor.Build.cs
@@ -13,7 +13,9 @@ namespace UnrealBuildTool.Rules
 					"Core",
 					"CoreUObject",
 					"DatasmithContent",
+					"DataprepCore",
 					"DesktopPlatform",
+					"EditorScriptingUtilities",
 					"EditorStyle",
 					"Engine",
 					"Projects",
diff --git a/Engine/Plugins/Enterprise/DatasmithContent/Source/DatasmithContentEditor/Private/AssetTypeActions_DatasmithScene.cpp b/Engine/Plugins/Enterprise/DatasmithContent/Source/DatasmithContentEditor/Private/AssetTypeActions_DatasmithScene.cpp
--- a/Engine/Plugins/Enterprise/DatasmithContent/Source/DatasmithContentEditor/Private/AssetTypeActions_DatasmithScene.cpp
+++ b/Engine/Plugins/Enterprise/DatasmithContent/Source/DatasmithContentEditor/Private/AssetTypeActions_DatasmithScene.cpp
@@ -4,7 +4,10 @@
 
 #include "DatasmithAssetImportData.h"
 #include "DatasmithContentEditorModule.h"
-#include "DatasmithScene.h"
+#include "DataprepAssetUserData.h"
+#include "DataprepCoreLibrary.h"
+
+#include "Framework/MultiBox/MultiBoxBuilder.h"
 
 #define LOCTEXT_NAMESPACE "AssetTypeActions_DatasmithScene"
 
@@ -61,4 +64,46 @@ void FAssetTypeActions_DatasmithScene::OpenAssetEditor(const TArray<UObject*>& I
 	}
 }
 
+void FAssetTypeActions_DatasmithScene::ExecuteDataprepAssets( TArray< TSoftObjectPtr< UDataprepAssetInterface > > DataprepAssets )
+{
+	for ( TSoftObjectPtr< UDataprepAssetInterface > DataprepAssetInterface : DataprepAssets ) 
+	{
+		UDataprepCoreLibrary::ExecuteWithReporting( DataprepAssetInterface.Get() );
+	}
+}
+
+void FAssetTypeActions_DatasmithScene::GetActions( const TArray<UObject*>& InObjects, FMenuBuilder& MenuBuilder )
+{
+	auto DatasmithScenes = GetTypedWeakObjectPtrs<UDatasmithScene>( InObjects );
+
+	if ( DatasmithScenes.Num() == 0 )
+	{
+		return;
+	}
+
+	// Allow execute only if we have at least one valid dataprep asset pointer
+	TArray< TSoftObjectPtr< UDataprepAssetInterface > > DataprepAssets;
+	for ( TWeakObjectPtr< UDatasmithScene > Scene : DatasmithScenes )
+	{
+		UDataprepAssetUserData* DataprepAssetUserData = CastChecked< UDataprepAssetUserData >( Scene->GetAssetUserDataOfClass( UDataprepAssetUserData::StaticClass() ), ECastCheckedType::NullAllowed );
+		if (DataprepAssetUserData && DataprepAssetUserData->DataprepAssetPtr)
+		{
+			DataprepAssets.Add( DataprepAssetUserData->DataprepAssetPtr );
+		}
+	}
+
+	if ( DataprepAssets.Num() > 0 )
+	{
+		MenuBuilder.AddMenuEntry(
+			LOCTEXT("RunAsset", "Execute"),
+			LOCTEXT("RunAssetTooltip", "Runs the Dataprep asset's producers, execute its recipe, finally runs the consumer"),
+			FSlateIcon(),
+			FUIAction(
+				FExecuteAction::CreateSP( this, &FAssetTypeActions_DatasmithScene::ExecuteDataprepAssets, DataprepAssets ),
+				FCanExecuteAction()
+			)
+		);
+	}
+}
+
 #undef LOCTEXT_NAMESPACE
diff --git a/Engine/Plugins/Enterprise/DatasmithContent/Source/DatasmithContentEditor/Private/AssetTypeActions_DatasmithScene.h b/Engine/Plugins/Enterprise/DatasmithContent/Source/DatasmithContentEditor/Private/AssetTypeActions_DatasmithScene.h
--- a/Engine/Plugins/Enterprise/DatasmithContent/Source/DatasmithContentEditor/Private/AssetTypeActions_DatasmithScene.h
+++ b/Engine/Plugins/Enterprise/DatasmithContent/Source/DatasmithContentEditor/Private/AssetTypeActions_DatasmithScene.h
@@ -3,6 +3,8 @@
 #pragma once
 
 #include "CoreMinimal.h"
+#include "DatasmithScene.h"
+#include "DataprepAssetInterface.h"
 #include "AssetTypeActions_Base.h"
 
 /** Asset type actions for UDatasmithScene class */
@@ -14,10 +16,13 @@ public:
 	virtual uint32 GetCategories() override;
 	virtual FColor GetTypeColor() const override { return FColor(255, 0, 0); }
 	virtual UClass* GetSupportedClass() const override;
-	virtual bool HasActions ( const TArray<UObject*>& InObjects ) const override { return false; }
+	virtual bool HasActions ( const TArray<UObject*>& InObjects ) const override { return true; }
+	virtual void GetActions( const TArray<UObject*>& InObjects, FMenuBuilder& MenuBuilder ) override;
 	virtual bool IsImportedAsset() const override { return true; }
 	virtual void GetResolvedSourceFilePaths(const TArray<UObject*>& TypeAssets, TArray<FString>& OutSourceFilePaths) const override;
 	virtual void OpenAssetEditor(const TArray<UObject*>& InObjects, TSharedPtr<class IToolkitHost> EditWithinLevelEditor = TSharedPtr<IToolkitHost>()) override;
 	// End IAssetTypeActions interface
-};
 
+private:
+	void ExecuteDataprepAssets(TArray< TSoftObjectPtr< UDataprepAssetInterface > > DataprepAssets);
+};
diff --git a/Engine/Plugins/Enterprise/DatasmithImporter/Source/DatasmithImporter/Private/DatasmithConsumer.cpp b/Engine/Plugins/Enterprise/DatasmithImporter/Source/DatasmithImporter/Private/DatasmithConsumer.cpp
--- a/Engine/Plugins/Enterprise/DatasmithImporter/Source/DatasmithImporter/Private/DatasmithConsumer.cpp
+++ b/Engine/Plugins/Enterprise/DatasmithImporter/Source/DatasmithImporter/Private/DatasmithConsumer.cpp
@@ -13,6 +13,8 @@
 #include "DatasmithSceneFactory.h"
 #include "IDatasmithSceneElements.h"
 #include "DatasmithStaticMeshImporter.h"
+#include "DataprepAssetUserData.h"
+#include "DataprepAssetInterface.h"
 #include "LevelVariantSets.h"
 #include "ObjectTemplates/DatasmithActorTemplate.h"
 #include "ObjectTemplates/DatasmithAreaLightActorTemplate.h"
@@ -111,6 +113,26 @@ bool UDatasmithConsumer::Initialize()
 
 		DatasmithScene->AssetImportData = NewObject< UDatasmithSceneImportData >( DatasmithScene.Get(), UDatasmithSceneImportData::StaticClass() );
 		check( DatasmithScene->AssetImportData );
+
+		// Store a dataprep asset pointer into the scene asset in order to be able to later re-execute the dataprep pipeline
+		UDataprepAssetInterface* DataprepAssetInterface = Cast< UDataprepAssetInterface >( GetOuter() );
+		check( DataprepAssetInterface );
+
+		if ( DatasmithScene->GetClass()->ImplementsInterface(UInterface_AssetUserData::StaticClass()) )
+		{
+			if ( IInterface_AssetUserData* AssetUserDataInterface = Cast< IInterface_AssetUserData >( DatasmithScene.Get() ) )
+			{
+				UDataprepAssetUserData* DataprepAssetUserData = AssetUserDataInterface->GetAssetUserData< UDataprepAssetUserData >();
+
+				if ( !DataprepAssetUserData )
+				{
+					EObjectFlags Flags = RF_Public;
+					DataprepAssetUserData = NewObject< UDataprepAssetUserData >( DatasmithScene.Get(), NAME_None, Flags );
+					AssetUserDataInterface->AddAssetUserData( DataprepAssetUserData );
+				}
+				DataprepAssetUserData->DataprepAssetPtr = DataprepAssetInterface;
+			}
+		}
 	}
 
 	// #ueent_todo: Find out necessity of namespace for uniqueness of asset's and actor's names
diff --git a/Engine/Plugins/Experimental/DataPrepEditor/Source/DataPrepCore/Private/DataPrepContentConsumer.cpp b/Engine/Plugins/Experimental/DataPrepEditor/Source/DataPrepCore/Private/DataPrepContentConsumer.cpp
--- a/Engine/Plugins/Experimental/DataPrepEditor/Source/DataPrepCore/Private/DataPrepContentConsumer.cpp
+++ b/Engine/Plugins/Experimental/DataPrepEditor/Source/DataPrepCore/Private/DataPrepContentConsumer.cpp
@@ -108,7 +108,7 @@ void UDataprepContentConsumer::AddDataprepAssetUserData()
 
 				if(!DataprepAssetUserData)
 				{
-					EObjectFlags Flags = RF_Public /*| RF_Transactional*/; // RF_Transactional Disabled as is can cause a crash in the transaction system for blueprints
+					EObjectFlags Flags = RF_Public /*| RF_Transactional*/; // RF_Transactional Disabled as it can cause a crash in the transaction system for blueprints
 
 					DataprepAssetUserData = NewObject< UDataprepAssetUserData >( Object, NAME_None, Flags );
 
diff --git a/Engine/Plugins/Experimental/DataPrepEditor/Source/DataPrepCore/Public/DataprepCoreLibrary.h b/Engine/Plugins/Experimental/DataPrepEditor/Source/DataPrepCore/Public/DataprepCoreLibrary.h
--- a/Engine/Plugins/Experimental/DataPrepEditor/Source/DataPrepCore/Public/DataprepCoreLibrary.h
+++ b/Engine/Plugins/Experimental/DataPrepEditor/Source/DataPrepCore/Public/DataprepCoreLibrary.h
@@ -12,7 +12,7 @@ class IDataprepLogger;
 class IDataprepProgressReporter;
 
 UCLASS(Blueprintable, BlueprintType, meta = (DisplayName = "Dataprep Core Blueprint Library"))
-class UDataprepCoreLibrary : public UBlueprintFunctionLibrary
+class DATAPREPCORE_API UDataprepCoreLibrary : public UBlueprintFunctionLibrary
 {
 	GENERATED_BODY()
 
