TIsPointerConvertibleFromTo usage removed from Datasmith overload resolution to avoid incorrect overloading when incomplete types are involved.

diff --git a/Engine/Plugins/Enterprise/DatasmithContent/Source/DatasmithContent/Public/DatasmithAdditionalData.h b/Engine/Plugins/Enterprise/DatasmithContent/Source/DatasmithContent/Public/DatasmithAdditionalData.h
--- a/Engine/Plugins/Enterprise/DatasmithContent/Source/DatasmithContent/Public/DatasmithAdditionalData.h
+++ b/Engine/Plugins/Enterprise/DatasmithContent/Source/DatasmithContent/Public/DatasmithAdditionalData.h
@@ -34,9 +34,11 @@ namespace Datasmith
 	 * @tparam DataType     Type of additional data to create. Must inherit from UDatasmithAdditionalData.
 	 * @return DataType*    Created instance
 	 */
-	template<class DataType>
-	inline typename TEnableIf<TPointerIsConvertibleFromTo<DataType, UDatasmithAdditionalData>::Value, DataType*>::Type
-	MakeAdditionalData()
+	template<
+		class DataType,
+		typename = decltype(ImplicitConv<UDatasmithAdditionalData*>((DataType*)nullptr))
+	>
+	inline DataType* MakeAdditionalData()
 	{
 		UClass* Class = DataType::StaticClass();
 		UObject* Outer = (UObject*)GetTransientPackage();
@@ -53,9 +55,11 @@ namespace Datasmith
 	 * @param MaxCount              Optional limitation of the result length
 	 * @return TArray<DataType*>    The list of additional data stored on the given asset
 	 */
-	template<class DataType>
-	inline typename TEnableIf<TPointerIsConvertibleFromTo<DataType, UDatasmithAdditionalData>::Value, TArray<DataType*>>::Type
-	GetMultipleAdditionalData(const FAssetData& SourceAssetData, int MaxCount=0)
+	template<
+		class DataType,
+		typename = decltype(ImplicitConv<UDatasmithAdditionalData*>((DataType*)nullptr))
+	>
+	inline TArray<DataType*> GetMultipleAdditionalData(const FAssetData& SourceAssetData, int MaxCount=0)
 	{
 		TArray<DataType*> Result;
 
@@ -87,9 +91,11 @@ namespace Datasmith
 	 * @param SourceAssetData   Asset on which the additional data is fetched
 	 * @return DataType*        nullptr if no Additional data is stored on the given asset
 	 */
-	template<class DataType>
-	inline typename TEnableIf<TPointerIsConvertibleFromTo<DataType, UDatasmithAdditionalData>::Value, DataType*>::Type
-	GetAdditionalData(const FAssetData& SourceAssetData)
+	template<
+		class DataType,
+		typename = decltype(ImplicitConv<UDatasmithAdditionalData*>((DataType*)nullptr))
+	>
+	inline DataType* GetAdditionalData(const FAssetData& SourceAssetData)
 	{
 		TArray<DataType*> Result = GetMultipleAdditionalData<DataType>(SourceAssetData, 1);
 		return Result.Num() ? Result[0] : nullptr;
