Merging CL 11372512 from //UE4/Dev-Enterprise to //UE4/Release-4.25

Fixed DatasmithMaterialInstanceTemplate to update material parameters when parent master material was changed

- also fixed regression for UE-81221, seems like a change(in DatasmithImporter.cpp) was overwritten by some merge(CL 10115208) later.

diff --git a/Engine/Plugins/Enterprise/DatasmithContent/Source/DatasmithContent/Private/ObjectTemplates/DatasmithMaterialInstanceTemplate.cpp b/Engine/Plugins/Enterprise/DatasmithContent/Source/DatasmithContent/Private/ObjectTemplates/DatasmithMaterialInstanceTemplate.cpp
--- a/Engine/Plugins/Enterprise/DatasmithContent/Source/DatasmithContent/Private/ObjectTemplates/DatasmithMaterialInstanceTemplate.cpp
+++ b/Engine/Plugins/Enterprise/DatasmithContent/Source/DatasmithContent/Private/ObjectTemplates/DatasmithMaterialInstanceTemplate.cpp
@@ -151,21 +151,26 @@ void FDatasmithStaticParameterSetTemplate::LoadRebase(const UMaterialInstanceCon
 
 	for (TMap< FName, bool >::TConstIterator It = ComparedTemplate.StaticSwitchParameters.CreateConstIterator(); It; ++It)
 	{
-		bool bDefaultValue;
+		bool DefaultValue;
 		FGuid ExpressionGuid;
-		if (Source.GetStaticSwitchParameterDefaultValue(It->Key, bDefaultValue, ExpressionGuid))
+		if (Source.GetStaticSwitchParameterDefaultValue(It->Key, DefaultValue, ExpressionGuid))
 		{
-			if (bDefaultValue != It->Value)
-			{
-				StaticSwitchParameters.Add(It->Key, It->Value);
-			}
-			else if (MergedTemplate)
+			bool OldValue = It->Value;
+			bool NewValue = DefaultValue;
+
+			if (MergedTemplate)
 			{
-				if (const bool* MergedTemplateValue = MergedTemplate->StaticSwitchParameters.Find(It->Key))
+				if (const bool* BaseTemplateValue = MergedTemplate->StaticSwitchParameters.Find(It->Key))
 				{
-					StaticSwitchParameters.Add(It->Key) = *MergedTemplateValue;
+					NewValue = *BaseTemplateValue;
 				}
 			}
+
+			// Store new value in case it's different from old or default is different from old
+			if (NewValue != OldValue || DefaultValue != OldValue)
+			{
+				StaticSwitchParameters.Add(It->Key) = NewValue;
+			}
 		}
 	}
 #endif // #if WITH_EDITORONLY_DATA
@@ -331,17 +336,22 @@ void UDatasmithMaterialInstanceTemplate::LoadRebase(const UObject* Source, const
 		float DefaultValue;
 		if (ParentMaterial->GetScalarParameterDefaultValue(It->Key, DefaultValue))
 		{
-			if (DefaultValue != It->Value)
-			{
-				ScalarParameterValues.Add(It->Key) = It->Value;
-			}
-			else if(bMergeTemplate)
+			float OldValue = It->Value;
+			float NewValue = DefaultValue;
+
+			if (bMergeTemplate)
 			{
 				if (const float* BaseTemplateValue = TypedBaseTemplate->ScalarParameterValues.Find(It->Key))
 				{
-					ScalarParameterValues.Add(It->Key) = *BaseTemplateValue;
+					NewValue = *BaseTemplateValue;
 				}
 			}
+
+			// Store new value in case it's different from old or default is different from old
+			if (NewValue != OldValue || DefaultValue != OldValue)
+			{
+				ScalarParameterValues.Add(It->Key) = NewValue;
+			}
 		}
 	}
 
@@ -351,17 +361,22 @@ void UDatasmithMaterialInstanceTemplate::LoadRebase(const UObject* Source, const
 		FLinearColor DefaultValue;
 		if (ParentMaterial->GetVectorParameterDefaultValue(It->Key, DefaultValue))
 		{
-			if (DefaultValue != It->Value)
-			{
-				VectorParameterValues.Add(It->Key) = It->Value;
-			}
-			else if (bMergeTemplate)
+			FLinearColor OldValue = It->Value;
+			FLinearColor NewValue = DefaultValue;
+			
+			if (bMergeTemplate)
 			{
 				if (const FLinearColor* BaseTemplateValue = TypedBaseTemplate->VectorParameterValues.Find(It->Key))
 				{
-					VectorParameterValues.Add(It->Key) = *BaseTemplateValue;
+					NewValue = *BaseTemplateValue;
 				}
 			}
+
+			// Store new value in case it's different from old or default is different from old
+			if (NewValue != OldValue || DefaultValue != OldValue)
+			{
+				VectorParameterValues.Add(It->Key) = NewValue;
+			}
 		}
 	}
 
@@ -371,17 +386,22 @@ void UDatasmithMaterialInstanceTemplate::LoadRebase(const UObject* Source, const
 		UTexture* DefaultValue;
 		if (ParentMaterial->GetTextureParameterDefaultValue(It->Key, DefaultValue))
 		{
-			if (DefaultValue != It->Value.Get())
-			{
-				TextureParameterValues.Add(It->Key) = It->Value;
-			}
-			else if (bMergeTemplate)
+			UTexture* OldValue = It->Value.Get();
+			UTexture* NewValue = DefaultValue;
+
+			if (bMergeTemplate)
 			{
 				if (const TSoftObjectPtr<UTexture>* BaseTemplateValue = TypedBaseTemplate->TextureParameterValues.Find(It->Key))
 				{
-					TextureParameterValues.Add(It->Key) = *BaseTemplateValue;
+					NewValue = BaseTemplateValue->Get();
 				}
 			}
+			
+			// Store new value in case it's different from old or default is different from old
+			if (NewValue != OldValue || DefaultValue != OldValue)
+			{
+				TextureParameterValues.Add(It->Key) = NewValue;
+			}
 		}
 	}
 
diff --git a/Engine/Plugins/Enterprise/DatasmithImporter/Source/DatasmithImporter/Private/DatasmithImporter.cpp b/Engine/Plugins/Enterprise/DatasmithImporter/Source/DatasmithImporter/Private/DatasmithImporter.cpp
--- a/Engine/Plugins/Enterprise/DatasmithImporter/Source/DatasmithImporter/Private/DatasmithImporter.cpp
+++ b/Engine/Plugins/Enterprise/DatasmithImporter/Source/DatasmithImporter/Private/DatasmithImporter.cpp
@@ -1257,9 +1257,9 @@ void FDatasmithImporter::ImportMaterials( FDatasmithImportContext& ImportContext
 
 		ImportContext.AssetsContext.MaterialsRequirements.Empty( ImportContext.FilteredScene->GetMaterialsCount() );
 
-		for ( int32 MaterialIndex = 0; MaterialIndex < ImportContext.FilteredScene->GetMaterialsCount(); ++MaterialIndex )
+		for (FDatasmithImporterUtils::FDatasmithMaterialImportIterator It(ImportContext); It; ++It)
 		{
-			TSharedRef< IDatasmithBaseMaterialElement > MaterialElement = ImportContext.FilteredScene->GetMaterial( MaterialIndex ).ToSharedRef();
+			TSharedRef< IDatasmithBaseMaterialElement > MaterialElement = It.Value().ToSharedRef();
 
 			UMaterialInterface* ExistingMaterial = nullptr;
 
