Fix serialization of UDatasmithCommonTessellationOptions

diff --git a/Engine/Plugins/Enterprise/DatasmithContent/Source/DatasmithContent/Public/DatasmithImportOptions.h b/Engine/Plugins/Enterprise/DatasmithContent/Source/DatasmithContent/Public/DatasmithImportOptions.h
--- a/Engine/Plugins/Enterprise/DatasmithContent/Source/DatasmithContent/Public/DatasmithImportOptions.h
+++ b/Engine/Plugins/Enterprise/DatasmithContent/Source/DatasmithContent/Public/DatasmithImportOptions.h
@@ -301,7 +301,7 @@ public:
 };
 
 
-UCLASS(BlueprintType, config = EditorPerProjectUserSettings, Transient)
+UCLASS(BlueprintType, config = EditorPerProjectUserSettings)
 class DATASMITHCONTENT_API UDatasmithCommonTessellationOptions : public UDatasmithOptionsBase
 {
 	GENERATED_BODY()
