Fixed PVS static analysis ternary operator warnings.

diff --git a/Engine/Plugins/Enterprise/DatasmithContent/Source/DatasmithContent/Private/ObjectTemplates/DatasmithStaticMeshComponentTemplate.cpp b/Engine/Plugins/Enterprise/DatasmithContent/Source/DatasmithContent/Private/ObjectTemplates/DatasmithStaticMeshComponentTemplate.cpp
--- a/Engine/Plugins/Enterprise/DatasmithContent/Source/DatasmithContent/Private/ObjectTemplates/DatasmithStaticMeshComponentTemplate.cpp
+++ b/Engine/Plugins/Enterprise/DatasmithContent/Source/DatasmithContent/Private/ObjectTemplates/DatasmithStaticMeshComponentTemplate.cpp
@@ -29,11 +29,11 @@ UObject* UDatasmithStaticMeshComponentTemplate::UpdateObject( UObject* Destinati
 	for ( int32 MaterialIndex = 0; MaterialIndex < OverrideMaterials.Num(); ++MaterialIndex )
 	{
 		UMaterialInterface* CurrentOverrideOnComponent = StaticMeshComponent->OverrideMaterials.IsValidIndex( MaterialIndex )
-			? StaticMeshComponent->OverrideMaterials[ MaterialIndex ]
+			? ToRawPtr(StaticMeshComponent->OverrideMaterials[ MaterialIndex ])
 			: nullptr;
 
 		UMaterialInterface* PreviousTemplateOverride = ( PreviousStaticMeshTemplate && PreviousStaticMeshTemplate->OverrideMaterials.IsValidIndex( MaterialIndex ) )
-			? PreviousStaticMeshTemplate->OverrideMaterials[ MaterialIndex ]
+			? ToRawPtr(PreviousStaticMeshTemplate->OverrideMaterials[ MaterialIndex ])
 			: nullptr;
 
 		UMaterialInterface* NewTemplateOverride = OverrideMaterials[ MaterialIndex ];
diff --git a/Engine/Plugins/VirtualProduction/RemoteControl/Source/RemoteControlCommon/Private/RCPropertyContainer.cpp b/Engine/Plugins/VirtualProduction/RemoteControl/Source/RemoteControlCommon/Private/RCPropertyContainer.cpp
--- a/Engine/Plugins/VirtualProduction/RemoteControl/Source/RemoteControlCommon/Private/RCPropertyContainer.cpp
+++ b/Engine/Plugins/VirtualProduction/RemoteControl/Source/RemoteControlCommon/Private/RCPropertyContainer.cpp
@@ -295,7 +295,7 @@ namespace PropertyContainers
 		FFieldVariant PropertyScope = InParent;
 
 		// A name is usually required further down the chain, so set one if not already.
-		const FName PropertyName = InPropertyName == NAME_None ? TEXT("Value") : InPropertyName;
+		const FName PropertyName = (InPropertyName == NAME_None) ? TEXT("Value") : InPropertyName;
 		
 		FProperty* NewProperty = nullptr;
 
