Datasmith: Fix static analysis warning

diff --git a/Engine/Plugins/Enterprise/DatasmithContent/Source/DatasmithContent/Private/ObjectTemplates/DatasmithStaticMeshComponentTemplate.cpp b/Engine/Plugins/Enterprise/DatasmithContent/Source/DatasmithContent/Private/ObjectTemplates/DatasmithStaticMeshComponentTemplate.cpp
--- a/Engine/Plugins/Enterprise/DatasmithContent/Source/DatasmithContent/Private/ObjectTemplates/DatasmithStaticMeshComponentTemplate.cpp
+++ b/Engine/Plugins/Enterprise/DatasmithContent/Source/DatasmithContent/Private/ObjectTemplates/DatasmithStaticMeshComponentTemplate.cpp
@@ -28,13 +28,17 @@ UObject* UDatasmithStaticMeshComponentTemplate::UpdateObject( UObject* Destinati
 
 	for ( int32 MaterialIndex = 0; MaterialIndex < OverrideMaterials.Num(); ++MaterialIndex )
 	{
-		UMaterialInterface* CurrentOverrideOnComponent = StaticMeshComponent->OverrideMaterials.IsValidIndex( MaterialIndex )
-			? StaticMeshComponent->OverrideMaterials[ MaterialIndex ]
-			: nullptr;
+		UMaterialInterface* CurrentOverrideOnComponent = nullptr;
+		if ( StaticMeshComponent->OverrideMaterials.IsValidIndex( MaterialIndex ) )
+		{
+			CurrentOverrideOnComponent = StaticMeshComponent->OverrideMaterials[ MaterialIndex ];
+		}
 
-		UMaterialInterface* PreviousTemplateOverride = ( PreviousStaticMeshTemplate && PreviousStaticMeshTemplate->OverrideMaterials.IsValidIndex( MaterialIndex ) )
-			? PreviousStaticMeshTemplate->OverrideMaterials[ MaterialIndex ]
-			: nullptr;
+		UMaterialInterface* PreviousTemplateOverride = nullptr;
+		if ( PreviousStaticMeshTemplate && PreviousStaticMeshTemplate->OverrideMaterials.IsValidIndex( MaterialIndex ) )
+		{
+			PreviousTemplateOverride = PreviousStaticMeshTemplate->OverrideMaterials[ MaterialIndex ];
+		}
 
 		UMaterialInterface* NewTemplateOverride = OverrideMaterials[ MaterialIndex ];
 
