Fix crash on exit on macOS.

#rnx

diff --git a/Engine/Plugins/Enterprise/DatasmithContent/Source/DatasmithContentEditor/Private/DatasmithContentEditorModule.cpp b/Engine/Plugins/Enterprise/DatasmithContent/Source/DatasmithContentEditor/Private/DatasmithContentEditorModule.cpp
--- a/Engine/Plugins/Enterprise/DatasmithContent/Source/DatasmithContentEditor/Private/DatasmithContentEditorModule.cpp
+++ b/Engine/Plugins/Enterprise/DatasmithContent/Source/DatasmithContentEditor/Private/DatasmithContentEditorModule.cpp
@@ -252,7 +252,7 @@ private:
 	{
 		const FName PropertyEditor("PropertyEditor");
 
-		if (FModuleManager::Get().IsModuleLoaded(PropertyEditor))
+		if (UObjectInitialized() && FModuleManager::Get().IsModuleLoaded(PropertyEditor))
 		{
 			FPropertyEditorModule& PropertyModule = FModuleManager::GetModuleChecked<FPropertyEditorModule>(PropertyEditor);
 			PropertyModule.UnregisterCustomPropertyTypeLayout(FDatasmithImportInfo::StaticStruct()->GetFName());
