Allow for ExternalSource Uri selection in DatasmithScene asset property.

Johan.Duparc

diff --git a/Engine/Plugins/Enterprise/DatasmithContent/Source/DatasmithContentEditor/DatasmithContentEditor.Build.cs b/Engine/Plugins/Enterprise/DatasmithContent/Source/DatasmithContentEditor/DatasmithContentEditor.Build.cs
--- a/Engine/Plugins/Enterprise/DatasmithContent/Source/DatasmithContentEditor/DatasmithContentEditor.Build.cs
+++ b/Engine/Plugins/Enterprise/DatasmithContent/Source/DatasmithContentEditor/DatasmithContentEditor.Build.cs
@@ -25,6 +25,7 @@ namespace UnrealBuildTool.Rules
 			PrivateDependencyModuleNames.AddRange(
 				new string[]
 				{
+					"InputCore",
 					"PropertyEditor",
 					"SlateCore",
 					"Slate",
diff --git a/Engine/Plugins/Enterprise/DatasmithContent/Source/DatasmithContentEditor/Private/DatasmithContentEditorModule.cpp b/Engine/Plugins/Enterprise/DatasmithContent/Source/DatasmithContentEditor/Private/DatasmithContentEditorModule.cpp
--- a/Engine/Plugins/Enterprise/DatasmithContent/Source/DatasmithContentEditor/Private/DatasmithContentEditorModule.cpp
+++ b/Engine/Plugins/Enterprise/DatasmithContent/Source/DatasmithContentEditor/Private/DatasmithContentEditorModule.cpp
@@ -135,7 +135,7 @@ public:
 
 	virtual void RegisterSetAssetAutoReimportHandler(FOnSetAssetAutoReimport&& SetAssetAutoReimportDelegate) override
 	{
-		SetAssetAutoReimportHandler = SetAssetAutoReimportDelegate;
+		SetAssetAutoReimportHandler = MoveTemp(SetAssetAutoReimportDelegate);
 	}
 	
 	virtual void UnregisterSetAssetAutoReimportHandler(FDelegateHandle InHandle) override
@@ -167,7 +167,7 @@ public:
 
 	virtual void RegisterIsAssetAutoReimportAvailableHandler(FOnIsAssetAutoReimportAvailable&& IsAssetAutoReimportAvailableDelegate)
 	{
-		IsAssetAutoReimportAvailableHandler = IsAssetAutoReimportAvailableDelegate;
+		IsAssetAutoReimportAvailableHandler = MoveTemp(IsAssetAutoReimportAvailableDelegate);
 	}
 	
 	virtual void UnregisterIsAssetAutoReimportAvailableHandler(FDelegateHandle InHandle)
@@ -190,7 +190,7 @@ public:
 
 	virtual void RegisterIsAssetAutoReimportEnabledHandler(FOnIsAssetAutoReimportEnabled&& IsAssetAutoReimportEnabledDelegate) override
 	{
-		IsAssetAutoReimportEnabledHandler = IsAssetAutoReimportEnabledDelegate;
+		IsAssetAutoReimportEnabledHandler = MoveTemp(IsAssetAutoReimportEnabledDelegate);
 	}
 	
 	virtual void UnregisterIsAssetAutoReimportEnabledHandler(FDelegateHandle InHandle) override
@@ -211,6 +211,51 @@ public:
 		return TOptional<bool>();
 	}
 
+	virtual void RegisterBrowseExternalSourceUriHandler(FOnBrowseExternalSourceUri&& BrowseExternalSourceUriDelegate) override
+	{
+		BrowseExternalSourceUriHandler = MoveTemp(BrowseExternalSourceUriDelegate);
+	}
+
+	virtual void UnregisterBrowseExternalSourceUriHandler(FDelegateHandle InHandle) override
+	{
+		if (BrowseExternalSourceUriHandler.GetHandle() == InHandle)
+		{
+			BrowseExternalSourceUriHandler.Unbind();
+		}
+	}
+
+	virtual bool BrowseExternalSourceUri(FName UriScheme, const FString& DefaultUri, FString& OutSourceUri, FString& OutFallbackFilepath) const override
+	{
+		if (IsAssetAutoReimportEnabledHandler.IsBound())
+		{
+			return BrowseExternalSourceUriHandler.Execute(UriScheme, DefaultUri, OutSourceUri, OutFallbackFilepath);
+		}
+
+		return false;
+	}
+
+	virtual void RegisterGetSupportedUriSchemeHandler(FOnGetSupportedUriSchemes&& GetSupportedUriSchemeDelegate) override
+	{
+		GetSupportedUriSchemeHandler = MoveTemp(GetSupportedUriSchemeDelegate);
+	}
+
+	virtual void UnregisterGetSupportedUriSchemeHandler(FDelegateHandle InHandle) override
+	{
+		if (GetSupportedUriSchemeHandler.GetHandle() == InHandle)
+		{
+			GetSupportedUriSchemeHandler.Unbind();
+		}
+	}
+
+	virtual TOptional<TArray<FName>> GetSupportedUriScheme() const override
+	{
+		if (GetSupportedUriSchemeHandler.IsBound())
+		{
+			return GetSupportedUriSchemeHandler.Execute();
+		}
+
+		return TOptional<TArray<FName>>();
+	}
 
 private:
 	static TSharedPtr<IDataprepImporterInterface> CreateEmptyDatasmithImportHandler()
@@ -268,6 +313,8 @@ private:
 	FOnSetAssetAutoReimport SetAssetAutoReimportHandler;
 	FOnIsAssetAutoReimportAvailable IsAssetAutoReimportAvailableHandler;
 	FOnIsAssetAutoReimportEnabled IsAssetAutoReimportEnabledHandler;
+	FOnBrowseExternalSourceUri BrowseExternalSourceUriHandler;
+	FOnGetSupportedUriSchemes GetSupportedUriSchemeHandler;
 
 	FDelegateHandle OnMapChangeHandle;
 	TSet<TSoftObjectPtr<UObject>> AutoReimportingAssets;
diff --git a/Engine/Plugins/Enterprise/DatasmithContent/Source/DatasmithContentEditor/Private/DatasmithContentEditorStyle.cpp b/Engine/Plugins/Enterprise/DatasmithContent/Source/DatasmithContentEditor/Private/DatasmithContentEditorStyle.cpp
--- a/Engine/Plugins/Enterprise/DatasmithContent/Source/DatasmithContentEditor/Private/DatasmithContentEditorStyle.cpp
+++ b/Engine/Plugins/Enterprise/DatasmithContent/Source/DatasmithContentEditor/Private/DatasmithContentEditorStyle.cpp
@@ -2,13 +2,14 @@
 
 #include "DatasmithContentEditorStyle.h"
 
-//#include "DatasmithContentEditorModule.h"
 #include "DatasmithContentModule.h"
 
 #include "Interfaces/IPluginManager.h"
+#include "Styling/CoreStyle.h"
 #include "Styling/SlateStyle.h"
 #include "Styling/SlateStyleRegistry.h"
 #include "Styling/SlateTypes.h"
+#include "Styling/StyleColors.h"
 
 #define IMAGE_PLUGIN_BRUSH( RelativePath, ... ) FSlateImageBrush( FDatasmithContentEditorStyle::InContent( RelativePath, ".png" ), __VA_ARGS__ )
 
@@ -16,6 +17,8 @@ TSharedPtr<FSlateStyleSet> FDatasmithContentEditorStyle::StyleSet;
 
 void FDatasmithContentEditorStyle::Initialize()
 {
+	using namespace CoreStyleConstants;
+
 	if (StyleSet.IsValid())
 	{
 		return;
@@ -26,10 +29,6 @@ void FDatasmithContentEditorStyle::Initialize()
 	StyleSet->SetContentRoot(FPaths::EngineContentDir() / TEXT("Editor/Slate"));
 	StyleSet->SetCoreContentRoot(FPaths::EngineContentDir() / TEXT("Slate"));
 
-	const FVector2D Icon20x16(16.0f, 16.0f);
-	const FVector2D Icon20x20(20.0f, 20.0f);
-	const FVector2D Icon40x40(40.0f, 40.0f);
-
 	StyleSet->Set("DatasmithDataprepEditor.Importer", new IMAGE_PLUGIN_BRUSH("Icons/DatasmithImporterIcon40", Icon40x40));
 	StyleSet->Set("DatasmithDataprepEditor.Importer.Small", new IMAGE_PLUGIN_BRUSH("Icons/DatasmithImporterIcon40", Icon20x20));
 	StyleSet->Set("DatasmithDataprepEditor.Importer.Selected", new IMAGE_PLUGIN_BRUSH("Icons/DatasmithImporterIcon40", Icon40x40));
@@ -72,6 +71,28 @@ void FDatasmithContentEditorStyle::Initialize()
 	StyleSet->Set("DatasmithDataprepEditor.Jacketing.Selected", new IMAGE_PLUGIN_BRUSH("Icons/Jacketing", Icon40x40));
 	StyleSet->Set("DatasmithDataprepEditor.Jacketing.Selected.Small", new IMAGE_PLUGIN_BRUSH("Icons/Jacketing", Icon20x20));
 
+	// Copy the base style, and change the rounding of the left-side borders so that it can be combined with a "right" widget.
+	FButtonStyle ButtonLeftStyle = FAppStyle::Get().GetWidgetStyle<FButtonStyle>(TEXT("Button"));
+	ButtonLeftStyle.SetNormal(FSlateRoundedBoxBrush(FStyleColors::Secondary, FVector4(4.0f, 0.0f, 0.0f, 4.0f), FStyleColors::Input, InputFocusThickness))
+		.SetHovered(FSlateRoundedBoxBrush(FStyleColors::Hover, FVector4(4.0f, 0.0f, 0.0f, 4.0f), FStyleColors::Input, InputFocusThickness))
+		.SetPressed(FSlateRoundedBoxBrush(FStyleColors::Header, FVector4(4.0f, 0.0f, 0.0f, 4.0f), FStyleColors::Input, InputFocusThickness))
+		.SetDisabled(FSlateRoundedBoxBrush(FStyleColors::Dropdown, FVector4(4.0f, 0.0f, 0.0f, 4.0f), FStyleColors::Recessed, InputFocusThickness));
+
+	StyleSet->Set("DatasmithDataprepEditor.ButtonLeft", ButtonLeftStyle);
+
+	// Copy the base style and tweak the padding and right-side border rounding so that it can be combined with a "left" widget.
+	FComboBoxStyle ComboRightStyle = FAppStyle::Get().GetWidgetStyle<FComboBoxStyle>(TEXT("SimpleComboBox"));
+	FComboButtonStyle& ComboButtonStyle = ComboRightStyle.ComboButtonStyle;
+	ComboButtonStyle.SetDownArrowPadding(1);
+	FButtonStyle& SimpleButtonStyle = ComboButtonStyle.ButtonStyle;
+	SimpleButtonStyle
+		.SetNormal(FSlateRoundedBoxBrush(FStyleColors::Secondary, FVector4(0.0f, 4.0f, 4.0f, 0.0f), FStyleColors::Input, InputFocusThickness))
+		.SetHovered(FSlateRoundedBoxBrush(FStyleColors::Hover, FVector4(0.0f, 4.0f, 4.0f, 0.0f), FStyleColors::Input, InputFocusThickness))
+		.SetPressed(FSlateRoundedBoxBrush(FStyleColors::Header, FVector4(0.0f, 4.0f, 4.0f, 0.0f), FStyleColors::Input, InputFocusThickness))
+		.SetDisabled(FSlateRoundedBoxBrush(FStyleColors::Dropdown, FVector4(0.0f, 4.0f, 4.0f, 0.0f), FStyleColors::Recessed, InputFocusThickness));
+
+	StyleSet->Set("DatasmithDataprepEditor.SimpleComboBoxRight", ComboRightStyle);
+
 	FSlateStyleRegistry::RegisterSlateStyle(*StyleSet.Get());
 }
 
diff --git a/Engine/Plugins/Enterprise/DatasmithContent/Source/DatasmithContentEditor/Private/DatasmithImportInfoCustomization.cpp b/Engine/Plugins/Enterprise/DatasmithContent/Source/DatasmithContentEditor/Private/DatasmithImportInfoCustomization.cpp
--- a/Engine/Plugins/Enterprise/DatasmithContent/Source/DatasmithContentEditor/Private/DatasmithImportInfoCustomization.cpp
+++ b/Engine/Plugins/Enterprise/DatasmithContent/Source/DatasmithContentEditor/Private/DatasmithImportInfoCustomization.cpp
@@ -3,6 +3,8 @@
 #include "DatasmithImportInfoCustomization.h"
 
 #include "DatasmithAssetImportData.h"
+#include "DatasmithContentEditorModule.h"
+#include "DatasmithContentEditorStyle.h"
 
 #include "DetailLayoutBuilder.h"
 #include "DetailWidgetRow.h"
@@ -10,14 +12,22 @@
 #include "PropertyHandle.h"
 #include "ScopedTransaction.h"
 #include "Widgets/DeclarativeSyntaxSupport.h"
+#include "Widgets/Input/SButton.h"
 #include "Widgets/Input/SEditableText.h"
 #include "Widgets/Input/SEditableTextBox.h"
 #include "Widgets/SBoxPanel.h"
 #include "Widgets/Text/STextBlock.h"
 
-
 #define LOCTEXT_NAMESPACE "AssetImportDataCustomization"
 
+FDatasmithImportInfoCustomization::FDatasmithImportInfoCustomization()
+{
+	if (TOptional<TArray<FName>> SupportedSchemes = IDatasmithContentEditorModule::Get().GetSupportedUriScheme())
+	{
+		AvailableUriSchemes = MoveTemp(SupportedSchemes.GetValue());
+	}
+}
+
 TSharedRef<IPropertyTypeCustomization> FDatasmithImportInfoCustomization::MakeInstance()
 {
 	return MakeShareable(new FDatasmithImportInfoCustomization);
@@ -27,15 +37,13 @@ void FDatasmithImportInfoCustomization::CustomizeChildren(TSharedRef<IPropertyHa
 {
 	PropertyHandle = InPropertyHandle;
 
-	FSlateFontInfo Font = IDetailLayoutBuilder::GetDetailFont();
-
-	FDatasmithImportInfo* Info = GetEditStruct();
-	if (Info)
+	if (GetImportInfo() != nullptr)
 	{
-		const FText SourceUriText = LOCTEXT("SourceUri", "Source Uri");
-		const FText SourceHashText = LOCTEXT("SourceHash", "Source Hash");
+		const FSlateFontInfo Font = IDetailLayoutBuilder::GetDetailFont();
+		const FText SourceUriLabel = LOCTEXT("SourceUri", "Source Uri");
 
-		FText SourceUriLabel = SourceUriText;
+		const FButtonStyle* ButtonStyle = &FDatasmithContentEditorStyle::Get()->GetWidgetStyle<FButtonStyle>(TEXT("DatasmithDataprepEditor.ButtonLeft"));
+		const FComboBoxStyle* ComboBoxStyle = &FDatasmithContentEditorStyle::Get()->GetWidgetStyle<FComboBoxStyle>(TEXT("DatasmithDataprepEditor.SimpleComboBoxRight"));
 
 		ChildBuilder.AddCustomRow(SourceUriLabel)
 		.NameContent()
@@ -51,21 +59,75 @@ void FDatasmithImportInfoCustomization::CustomizeChildren(TSharedRef<IPropertyHa
 			SNew(SHorizontalBox)
 			+ SHorizontalBox::Slot()
 			.VAlign(VAlign_Center)
+			.FillWidth(1.f)
 			[
-				SNew(SEditableTextBox)
-				.IsReadOnly(false)
+				SNew(SEditableText)
+				.IsReadOnly(true)
 				.Text(this, &FDatasmithImportInfoCustomization::GetUriText)
 				.ToolTipText(this, &FDatasmithImportInfoCustomization::GetUriText)
-				.OnTextCommitted(this, &FDatasmithImportInfoCustomization::OnSourceUriChanged)
 				.Font(Font)
 			]
+
+			+ SHorizontalBox::Slot()
+			.VAlign(VAlign_Center)
+			.AutoWidth()
+			[
+				SNew(SButton)
+				.ButtonStyle(ButtonStyle)
+				.OnClicked(this, &FDatasmithImportInfoCustomization::OnBrowseSourceClicked)
+				.ToolTipText(LOCTEXT("ChangePath_Tooltip", "Browse for a new source file path"))
+				[
+					SNew(STextBlock)
+					.Text(LOCTEXT("...", "..."))
+					.Font(Font)
+				]
+			]
+
+			+ SHorizontalBox::Slot()
+			.VAlign(VAlign_Center)
+			.AutoWidth()
+			[
+				SAssignNew(SchemeComboBox, SComboBox<FName>)
+				.OptionsSource(&AvailableUriSchemes)
+				.OnSelectionChanged(this, &FDatasmithImportInfoCustomization::OnSchemeComboBoxChanged)
+				.OnGenerateWidget_Lambda([](FName InItem) { return SNew(STextBlock).Text(FText::FromName(InItem)); })
+				.ComboBoxStyle(ComboBoxStyle)
+				[
+					// We only need to display the "down arrow" but SComboBox does not accept SNullWidget as valid content.
+					// Use an empty STextBlock instead.
+					SNew(STextBlock)
+				]
+			]
 		];
 	}
 }
 
+FName FDatasmithImportInfoCustomization::GetCurrentScheme() const
+{
+	FName CurrentScheme;
+
+	if (FDatasmithImportInfo* Info = GetImportInfo())
+	{
+		// #ueent_todo Once source URI become part of the core API, update this to FSourceUri::GetScheme();
+		int32 SchemeSeparatorIndex = Info->SourceUri.Find("://");
+		if (SchemeSeparatorIndex != INDEX_NONE)
+		{
+			CurrentScheme = FName(Info->SourceUri.Left(SchemeSeparatorIndex));
+		}
+	}
+
+	if (CurrentScheme.IsNone())
+	{
+		// we should fallback on file selection if we can't determine the scheme.
+		CurrentScheme = FName(TEXT("file"));
+	}
+
+	return CurrentScheme;
+}
+
 FText FDatasmithImportInfoCustomization::GetUriText() const
 {
-	FDatasmithImportInfo* Info = GetEditStruct();
+	FDatasmithImportInfo* Info = GetImportInfo();
 	if (Info)
 	{
 		return FText::FromString(Info->SourceUri);
@@ -73,7 +135,7 @@ FText FDatasmithImportInfoCustomization::GetUriText() const
 	return LOCTEXT("NoUriFound", "No Source Uri Set");
 }
 
-FDatasmithImportInfo* FDatasmithImportInfoCustomization::GetEditStruct() const
+FDatasmithImportInfo* FDatasmithImportInfoCustomization::GetImportInfo() const
 {
 	TArray<FDatasmithImportInfo*> AssetImportInfo;
 
@@ -89,7 +151,6 @@ FDatasmithImportInfo* FDatasmithImportInfoCustomization::GetEditStruct() const
 	return nullptr;
 }
 
-
 UObject* FDatasmithImportInfoCustomization::GetOuterClass() const
 {
 	static TArray<UObject*> OuterObjects;
@@ -100,7 +161,6 @@ UObject* FDatasmithImportInfoCustomization::GetOuterClass() const
 	return OuterObjects.Num() ? OuterObjects[0] : nullptr;
 }
 
-
 class FImportDataSourceFileTransactionScope
 {
 public:
@@ -133,23 +193,53 @@ private:
 	UObject* OuterObject;
 };
 
-void FDatasmithImportInfoCustomization::OnSourceUriChanged(const FText& NewText, ETextCommit::Type) const
+FReply FDatasmithImportInfoCustomization::OnBrowseSourceClicked() const
+{
+	SelectNewSource(GetCurrentScheme());
+
+	return FReply::Handled();
+}
+
+bool FDatasmithImportInfoCustomization::SelectNewSource(FName SourceScheme) const
 {
-	UObject* OuterObject = GetOuterClass();
-	FDatasmithImportInfo* Info = GetEditStruct();
+	UAssetImportData* ImportData = Cast<UAssetImportData>(GetOuterClass());
+	UObject* Obj = ImportData ? ImportData->GetOuter() : nullptr;
+
+	FDatasmithImportInfo* Info = GetImportInfo();
+	if (!Obj || !Info)
+	{
+		return false;
+	}
 
-	if (Info && OuterObject)
+	FString SourceUri;
+	FString FallbackFilepath;
+	if (IDatasmithContentEditorModule::Get().BrowseExternalSourceUri(SourceScheme, Info->SourceUri, SourceUri, FallbackFilepath))
 	{
-		FImportDataSourceFileTransactionScope TransactionScope(LOCTEXT("SourceUriChanged", "Change Source URI"), OuterObject);
+		FImportDataSourceFileTransactionScope TransactionScope(LOCTEXT("SourceUriChanged", "Change Source URI"), ImportData);
 
-		Info->SourceUri = NewText.ToString();
+		Info->SourceUri = SourceUri;
 		Info->SourceHash.Empty();
+		ImportData->UpdateFilenameOnly(FallbackFilepath);
 
 		// Broadcasting property change to force refresh the asset registry tag and notify systems monitoring the URI.
 		FPropertyChangedEvent EmptyPropertyChangedEvent(nullptr);
-		FCoreUObjectDelegates::OnObjectPropertyChanged.Broadcast(OuterObject, EmptyPropertyChangedEvent);
+		FCoreUObjectDelegates::OnObjectPropertyChanged.Broadcast(ImportData, EmptyPropertyChangedEvent);
+		return true;
 	}
+
+	return false;
 }
 
+void FDatasmithImportInfoCustomization::OnSchemeComboBoxChanged(FName InItem, ESelectInfo::Type InSeletionInfo)
+{
+	if (!InItem.IsNone())
+	{
+		// User must select a new source when changing the scheme (type) of the source.
+		SelectNewSource(InItem);
+
+		// If we don't clear the selection OnSchemeComboBoxChanged() won't be triggered if the used select the same option again.
+		SchemeComboBox->ClearSelection();
+	}
+}
 
 #undef LOCTEXT_NAMESPACE
diff --git a/Engine/Plugins/Enterprise/DatasmithContent/Source/DatasmithContentEditor/Private/DatasmithImportInfoCustomization.h b/Engine/Plugins/Enterprise/DatasmithContent/Source/DatasmithContentEditor/Private/DatasmithImportInfoCustomization.h
--- a/Engine/Plugins/Enterprise/DatasmithContent/Source/DatasmithContentEditor/Private/DatasmithImportInfoCustomization.h
+++ b/Engine/Plugins/Enterprise/DatasmithContent/Source/DatasmithContentEditor/Private/DatasmithImportInfoCustomization.h
@@ -3,8 +3,10 @@
 #pragma once
 
 #include "CoreMinimal.h"
+#include "Containers/Array.h"
 #include "Input/Reply.h"
 #include "IPropertyTypeCustomization.h"
+#include "Widgets/Input/SComboBox.h"
 
 class FDetailWidgetRow;
 class IDetailChildrenBuilder;
@@ -15,6 +17,8 @@ struct FDatasmithImportInfo;
 class FDatasmithImportInfoCustomization : public IPropertyTypeCustomization
 {
 public:
+	FDatasmithImportInfoCustomization();
+
 	/** Makes a new instance of this detail layout class for a specific detail view requesting it */
 	static TSharedRef<IPropertyTypeCustomization> MakeInstance();
 
@@ -24,10 +28,13 @@ public:
 
 private:
 	
-	void OnSourceUriChanged(const FText& NewText, ETextCommit::Type) const;
+	FName GetCurrentScheme() const;
+
+	FReply OnBrowseSourceClicked() const;
+
+	bool SelectNewSource(FName SourceScheme) const;
 
-	/** Access the struct we are editing - returns null if we have more than one. */
-	FDatasmithImportInfo* GetEditStruct() const;
+	FDatasmithImportInfo* GetImportInfo() const;
 
 	/** Access the outer class that contains this struct */
 	UObject* GetOuterClass() const;
@@ -35,7 +42,13 @@ private:
 	/** Get text for the UI */
 	FText GetUriText() const;
 
+	void OnSchemeComboBoxChanged(FName InItem, ESelectInfo::Type InSeletionInfo);
+
 private:
 	/** Property handle of the property we're editing */
 	TSharedPtr<IPropertyHandle> PropertyHandle;
+
+	TSharedPtr<SComboBox<FName>> SchemeComboBox;
+
+	TArray<FName> AvailableUriSchemes;
 };
\ No newline at end of file
diff --git a/Engine/Plugins/Enterprise/DatasmithContent/Source/DatasmithContentEditor/Public/DatasmithContentEditorModule.h b/Engine/Plugins/Enterprise/DatasmithContent/Source/DatasmithContentEditor/Public/DatasmithContentEditorModule.h
--- a/Engine/Plugins/Enterprise/DatasmithContent/Source/DatasmithContentEditor/Public/DatasmithContentEditorModule.h
+++ b/Engine/Plugins/Enterprise/DatasmithContent/Source/DatasmithContentEditor/Public/DatasmithContentEditorModule.h
@@ -41,6 +41,8 @@ DECLARE_DELEGATE_RetVal(TSharedPtr<IDataprepImporterInterface>, FOnCreateDatasmi
 DECLARE_DELEGATE_RetVal_TwoParams(bool, FOnSetAssetAutoReimport, UObject* /*Asset*/, bool /*bEnabled*/);
 DECLARE_DELEGATE_RetVal_OneParam(bool, FOnIsAssetAutoReimportAvailable, UObject* /*Asset*/);
 DECLARE_DELEGATE_RetVal_OneParam(bool, FOnIsAssetAutoReimportEnabled, UObject* /*Asset*/);
+DECLARE_DELEGATE_RetVal_FourParams(bool, FOnBrowseExternalSourceUri, FName/*UriScheme*/, const FString& /*DefaultUri*/, FString& /*OutSourceUri*/, FString& /*OutFallbackFilepath*/);
+DECLARE_DELEGATE_RetVal(const TArray<FName>&, FOnGetSupportedUriSchemes);
 
 struct FImporterDescription
 {
@@ -126,5 +128,13 @@ public:
 	virtual void RegisterIsAssetAutoReimportEnabledHandler(FOnIsAssetAutoReimportEnabled&& IsAssetAutoReimportEnabledDelegate) = 0;
 	virtual void UnregisterIsAssetAutoReimportEnabledHandler(FDelegateHandle InHandle) = 0;
 	virtual TOptional<bool> IsAssetAutoReimportEnabled(UObject* Asset) const = 0;
+
+	virtual void RegisterBrowseExternalSourceUriHandler(FOnBrowseExternalSourceUri&& BrowseExternalSourceUriDelegate) = 0;
+	virtual void UnregisterBrowseExternalSourceUriHandler(FDelegateHandle InHandle) = 0;
+	virtual bool BrowseExternalSourceUri(FName UriScheme, const FString& DefaultUri, FString& OutSourceUri, FString& OutFallbackFilepath) const = 0;
+
+	virtual void RegisterGetSupportedUriSchemeHandler(FOnGetSupportedUriSchemes&& GetSupportedUriSchemeDelegate) = 0;
+	virtual void UnregisterGetSupportedUriSchemeHandler(FDelegateHandle InHandle) = 0;
+	virtual TOptional<TArray<FName>> GetSupportedUriScheme() const = 0;
 };
 
diff --git a/Engine/Plugins/Enterprise/DatasmithImporter/Source/DatasmithExternalSource/DatasmithExternalSource.Build.cs b/Engine/Plugins/Enterprise/DatasmithImporter/Source/DatasmithExternalSource/DatasmithExternalSource.Build.cs
--- a/Engine/Plugins/Enterprise/DatasmithImporter/Source/DatasmithExternalSource/DatasmithExternalSource.Build.cs
+++ b/Engine/Plugins/Enterprise/DatasmithImporter/Source/DatasmithExternalSource/DatasmithExternalSource.Build.cs
@@ -22,8 +22,20 @@ namespace UnrealBuildTool.Rules
                     "DatasmithCore",
                     "DatasmithTranslator",
 					"Projects",
-                }
-            );
-        }
+				}
+			);
+
+			if (Target.Type == TargetType.Editor)
+			{
+				PrivateDependencyModuleNames.AddRange(
+					new string[]
+					{
+						"DesktopPlatform",
+						"Slate",
+						"UnrealEd",
+					}
+				);
+			}
+		}
     }
 }
diff --git a/Engine/Plugins/Enterprise/DatasmithImporter/Source/DatasmithExternalSource/Private/DatasmithFileUriResolver.cpp b/Engine/Plugins/Enterprise/DatasmithImporter/Source/DatasmithExternalSource/Private/DatasmithFileUriResolver.cpp
--- a/Engine/Plugins/Enterprise/DatasmithImporter/Source/DatasmithExternalSource/Private/DatasmithFileUriResolver.cpp
+++ b/Engine/Plugins/Enterprise/DatasmithImporter/Source/DatasmithExternalSource/Private/DatasmithFileUriResolver.cpp
@@ -6,6 +6,13 @@
 #include "DatasmithSceneSource.h"
 #include "DatasmithTranslatorManager.h"
 
+#if WITH_EDITOR
+#include "DesktopPlatformModule.h"
+#include "Framework/Application/SlateApplication.h"
+#include "IDesktopPlatform.h"
+#include "ObjectTools.h"
+#endif //WITH_EDITOR
+
 namespace UE::DatasmithImporter
 {
 	TSharedPtr<FExternalSource> FDatasmithFileUriResolver::GetOrCreateExternalSource(const FSourceUri& Uri) const
@@ -34,4 +41,51 @@ namespace UE::DatasmithImporter
 	{
 		return FName(FSourceUri::GetFileScheme());
 	}
+
+#if WITH_EDITOR
+	TSharedPtr<FExternalSource> FDatasmithFileUriResolver::BrowseExternalSource(const FSourceUri& DefaultUri) const
+	{
+		const TArray<FString> Formats = FDatasmithTranslatorManager::Get().GetSupportedFormats();
+		FString FileTypes;
+		FString Extensions;
+
+		ObjectTools::AppendFormatsFileExtensions(Formats, FileTypes, Extensions);
+
+		const FString FormatString = FString::Printf(TEXT("All Files (%s)|%s|%s"), *Extensions, *Extensions, *FileTypes);
+		const FString Title = NSLOCTEXT("DatasmithFileUriResolver", "BrowseSourceDialogTitle", "Select Datasmith Source File").ToString();
+
+		FString DefaultFolder;
+		FString DefaultFile;
+		if (DefaultUri.IsValid())
+		{
+			const FString DefaultFilePath(DefaultUri.GetPath());
+			DefaultFolder = FPaths::GetPath(FString(DefaultFilePath));
+			DefaultFile = FPaths::GetCleanFilename(DefaultFilePath);
+		}
+
+		TArray<FString> OutOpenFilenames;
+		IDesktopPlatform* DesktopPlatform = FDesktopPlatformModule::Get();
+		bool bOpened = false;
+		if (DesktopPlatform)
+		{
+			bOpened = DesktopPlatform->OpenFileDialog(
+				FSlateApplication::Get().FindBestParentWindowHandleForDialogs(nullptr),
+				Title,
+				DefaultFolder,
+				DefaultFile,
+				FormatString,
+				EFileDialogFlags::None,
+				OutOpenFilenames
+			);
+		}
+
+		if (bOpened && OutOpenFilenames.Num() > 0)
+		{
+			const FSourceUri FileSourceUri = FSourceUri::FromFilePath(OutOpenFilenames[0]);
+			return GetOrCreateExternalSource(FileSourceUri);
+		}
+
+		return nullptr;
+	}
+#endif //WITH_EDITOR
 }
\ No newline at end of file
diff --git a/Engine/Plugins/Enterprise/DatasmithImporter/Source/DatasmithExternalSource/Private/DatasmithFileUriResolver.h b/Engine/Plugins/Enterprise/DatasmithImporter/Source/DatasmithExternalSource/Private/DatasmithFileUriResolver.h
--- a/Engine/Plugins/Enterprise/DatasmithImporter/Source/DatasmithExternalSource/Private/DatasmithFileUriResolver.h
+++ b/Engine/Plugins/Enterprise/DatasmithImporter/Source/DatasmithExternalSource/Private/DatasmithFileUriResolver.h
@@ -14,6 +14,9 @@ namespace UE::DatasmithImporter
 		virtual TSharedPtr<FExternalSource> GetOrCreateExternalSource(const FSourceUri& URI) const override;
 		virtual bool CanResolveUri(const FSourceUri& URI) const override;
 		virtual FName GetScheme() const override;
+#if WITH_EDITOR
+		virtual TSharedPtr<FExternalSource> BrowseExternalSource(const FSourceUri& DefaultUri) const override;
+#endif //WITH_EDITOR
 		// IUriResolver interface end
 	};
 }
\ No newline at end of file
diff --git a/Engine/Plugins/Enterprise/DatasmithImporter/Source/DatasmithImporter/Private/DatasmithImporterModule.cpp b/Engine/Plugins/Enterprise/DatasmithImporter/Source/DatasmithImporter/Private/DatasmithImporterModule.cpp
--- a/Engine/Plugins/Enterprise/DatasmithImporter/Source/DatasmithImporter/Private/DatasmithImporterModule.cpp
+++ b/Engine/Plugins/Enterprise/DatasmithImporter/Source/DatasmithImporter/Private/DatasmithImporterModule.cpp
@@ -167,6 +167,8 @@ private:
 	FDelegateHandle SetAssetAutoReimportHandle;
 	FDelegateHandle IsAssetAutoReimportSupportedHandle;
 	FDelegateHandle IsAssetAutoReimportEnabledHandle;
+	FDelegateHandle BrowseExternalSourceUriHandle;
+	FDelegateHandle GetSupporedUriSchemeHandle;
 };
 
 void FDatasmithImporterModule::SetupMenuEntry()
@@ -793,6 +795,39 @@ void FDatasmithImporterModule::SetupDatasmithContentDelegates()
 		IsAssetAutoReimportEnabledHandle = IsAssetAutoReimportEnabled.GetHandle();
 		DatasmithContentEditorModule.RegisterIsAssetAutoReimportEnabledHandler(MoveTemp(IsAssetAutoReimportEnabled));
 	}
+
+	{
+		FOnBrowseExternalSourceUri BrowseExternalSourceUri = FOnBrowseExternalSourceUri::CreateLambda(
+			[](FName UriScheme, const FString& DefaultUriString, FString& OutSourceUri, FString& OutFallbackFilepath)
+			{
+				using namespace UE::DatasmithImporter;
+				FSourceUri DefaultUri(DefaultUriString);
+
+				if (TSharedPtr<FExternalSource> ExternalSource = IExternalSourceModule::Get().GetManager().BrowseExternalSource(UriScheme, DefaultUri))
+				{
+					OutSourceUri = ExternalSource->GetSourceUri().ToString();
+					OutFallbackFilepath = ExternalSource->GetFallbackFilepath();
+					return true;
+				}
+
+				return false;
+			});
+
+		BrowseExternalSourceUriHandle = BrowseExternalSourceUri.GetHandle();
+		DatasmithContentEditorModule.RegisterBrowseExternalSourceUriHandler(MoveTemp(BrowseExternalSourceUri));
+	}
+
+	{
+		FOnGetSupportedUriSchemes GetSupporedUriScheme = FOnGetSupportedUriSchemes::CreateLambda(
+			[]() -> const TArray<FName>&
+			{
+				using namespace UE::DatasmithImporter;
+				return IExternalSourceModule::Get().GetManager().GetSupportedSchemes();
+			});
+
+		GetSupporedUriSchemeHandle = GetSupporedUriScheme.GetHandle();
+		DatasmithContentEditorModule.RegisterGetSupportedUriSchemeHandler(MoveTemp(GetSupporedUriScheme));
+	}
 }
 
 void FDatasmithImporterModule::RemoveDatasmithContentDelegates()
@@ -804,6 +839,8 @@ void FDatasmithImporterModule::RemoveDatasmithContentDelegates()
 		DatasmithContentEditorModule.UnregisterSetAssetAutoReimportHandler(SetAssetAutoReimportHandle);
 		DatasmithContentEditorModule.UnregisterIsAssetAutoReimportAvailableHandler(IsAssetAutoReimportSupportedHandle);
 		DatasmithContentEditorModule.UnregisterIsAssetAutoReimportEnabledHandler(IsAssetAutoReimportEnabledHandle);
+		DatasmithContentEditorModule.UnregisterBrowseExternalSourceUriHandler(BrowseExternalSourceUriHandle);
+		DatasmithContentEditorModule.UnregisterGetSupportedUriSchemeHandler(GetSupporedUriSchemeHandle);
 	}
 }
 
diff --git a/Engine/Plugins/Enterprise/DatasmithImporter/Source/DirectLinkExtension/Private/DirectLinkUriResolver.cpp b/Engine/Plugins/Enterprise/DatasmithImporter/Source/DirectLinkExtension/Private/DirectLinkUriResolver.cpp
--- a/Engine/Plugins/Enterprise/DatasmithImporter/Source/DirectLinkExtension/Private/DirectLinkUriResolver.cpp
+++ b/Engine/Plugins/Enterprise/DatasmithImporter/Source/DirectLinkExtension/Private/DirectLinkUriResolver.cpp
@@ -52,6 +52,13 @@ namespace UE::DatasmithImporter
 		return TOptional<FDirectLinkSourceDescription>();
 	}
 
+#if WITH_EDITOR
+	TSharedPtr<FExternalSource> FDirectLinkUriResolver::BrowseExternalSource(const FSourceUri& DefaultSourceUri) const
+	{
+		return IDirectLinkExtensionModule::Get().DisplayDirectLinkSourcesDialog();
+	}
+#endif //WITH_EDITOR
+
 	const FString& FDirectLinkUriResolver::GetDirectLinkScheme()
 	{
 		static FString Scheme(TEXT("directlink"));
diff --git a/Engine/Plugins/Enterprise/DatasmithImporter/Source/DirectLinkExtension/Public/DirectLinkUriResolver.h b/Engine/Plugins/Enterprise/DatasmithImporter/Source/DirectLinkExtension/Public/DirectLinkUriResolver.h
--- a/Engine/Plugins/Enterprise/DatasmithImporter/Source/DirectLinkExtension/Public/DirectLinkUriResolver.h
+++ b/Engine/Plugins/Enterprise/DatasmithImporter/Source/DirectLinkExtension/Public/DirectLinkUriResolver.h
@@ -29,6 +29,9 @@ namespace UE::DatasmithImporter
 		virtual TSharedPtr<FExternalSource> GetOrCreateExternalSource(const FSourceUri& Uri) const override;
 		virtual bool CanResolveUri(const FSourceUri& Uri) const override;
 		virtual FName GetScheme() const override { return FName(GetDirectLinkScheme()); }
+#if WITH_EDITOR
+		virtual TSharedPtr<FExternalSource> BrowseExternalSource(const FSourceUri& DefaultSourceUri) const override;
+#endif //WITH_EDITOR
 		// IUriResolver interface end
 
 		/**
diff --git a/Engine/Plugins/Enterprise/DatasmithImporter/Source/ExternalSource/Private/UriManager.cpp b/Engine/Plugins/Enterprise/DatasmithImporter/Source/ExternalSource/Private/UriManager.cpp
--- a/Engine/Plugins/Enterprise/DatasmithImporter/Source/ExternalSource/Private/UriManager.cpp
+++ b/Engine/Plugins/Enterprise/DatasmithImporter/Source/ExternalSource/Private/UriManager.cpp
@@ -40,7 +40,7 @@ namespace UE::DatasmithImporter
 		return NumberOfEntryRemoved > 0;
 	}
 
-	const TArray<FString>& FUriManager::GetSupportedSchemes() const
+	const TArray<FName>& FUriManager::GetSupportedSchemes() const
 	{
 		if (CachedSchemes.Num() == 0 && RegisteredResolvers.Num() > 0)
 		{
@@ -48,7 +48,7 @@ namespace UE::DatasmithImporter
 
 			for (const FUriResolverRegisterInformation& RegisterInfo : RegisteredResolvers)
 			{
-				CachedSchemes.Add(RegisterInfo.UriResolver->GetScheme().ToString());
+				CachedSchemes.Add(RegisterInfo.UriResolver->GetScheme());
 			}
 		}
 
@@ -72,4 +72,19 @@ namespace UE::DatasmithImporter
 	{
 		CachedSchemes.Reset();
 	}
+
+#if WITH_EDITOR
+	TSharedPtr<FExternalSource> FUriManager::BrowseExternalSource(const FName& UriScheme, const FSourceUri& DefaultSourceUri) const
+	{
+		for (const FUriResolverRegisterInformation& RegisterInfo : RegisteredResolvers)
+		{
+			if (RegisterInfo.UriResolver->GetScheme() == UriScheme)
+			{
+				return RegisterInfo.UriResolver->BrowseExternalSource(DefaultSourceUri);
+			}
+		}
+
+		return nullptr;
+	}
+#endif //WITH_EDITOR
 }
\ No newline at end of file
diff --git a/Engine/Plugins/Enterprise/DatasmithImporter/Source/ExternalSource/Private/UriManager.h b/Engine/Plugins/Enterprise/DatasmithImporter/Source/ExternalSource/Private/UriManager.h
--- a/Engine/Plugins/Enterprise/DatasmithImporter/Source/ExternalSource/Private/UriManager.h
+++ b/Engine/Plugins/Enterprise/DatasmithImporter/Source/ExternalSource/Private/UriManager.h
@@ -28,7 +28,10 @@ namespace UE::DatasmithImporter
 		virtual bool CanResolveUri(const FSourceUri& URI) const override;
 		virtual void RegisterResolver(FName ResolverName, const TSharedRef<IUriResolver>& UriResolver) override;
 		virtual bool UnregisterResolver(FName ResolverName) override;
-		virtual const TArray<FString>& GetSupportedSchemes() const override;
+		virtual const TArray<FName>& GetSupportedSchemes() const override;
+#if WITH_EDITOR
+		virtual TSharedPtr<FExternalSource> BrowseExternalSource(const FName& UriScheme, const FSourceUri& DefaultSourceUri) const override;
+#endif //WITH_EDITOR
 		// IUriManager interface end
 
 	private:
@@ -38,7 +41,7 @@ namespace UE::DatasmithImporter
 
 		TArray<FUriResolverRegisterInformation> RegisteredResolvers;
 
-		mutable TArray<FString> CachedSchemes;
+		mutable TArray<FName> CachedSchemes;
 	};
 	
 }
\ No newline at end of file
diff --git a/Engine/Plugins/Enterprise/DatasmithImporter/Source/ExternalSource/Public/IUriManager.h b/Engine/Plugins/Enterprise/DatasmithImporter/Source/ExternalSource/Public/IUriManager.h
--- a/Engine/Plugins/Enterprise/DatasmithImporter/Source/ExternalSource/Public/IUriManager.h
+++ b/Engine/Plugins/Enterprise/DatasmithImporter/Source/ExternalSource/Public/IUriManager.h
@@ -47,9 +47,14 @@ namespace UE::DatasmithImporter
 		/**
 		 * Return an array containing all the scheme supported by the registered URiResolvers.
 		 */
-		virtual const TArray<FString>& GetSupportedSchemes() const = 0;
+		virtual const TArray<FName>& GetSupportedSchemes() const = 0;
 
 #if WITH_EDITOR
+		/**
+		 * Open a dialog window for browsing ExternalSource based on the passed URI Scheme
+		 */
+		virtual TSharedPtr<FExternalSource> BrowseExternalSource(const FName& UriScheme, const FSourceUri& DefaultSourceUri = FSourceUri()) const = 0;
+
 		template<typename ImportDataType>
 		TSharedPtr<FExternalSource> TryGetExternalSourceFromImportData(const ImportDataType& ImportSourceData) const
 		{
diff --git a/Engine/Plugins/Enterprise/DatasmithImporter/Source/ExternalSource/Public/IUriResolver.h b/Engine/Plugins/Enterprise/DatasmithImporter/Source/ExternalSource/Public/IUriResolver.h
--- a/Engine/Plugins/Enterprise/DatasmithImporter/Source/ExternalSource/Public/IUriResolver.h
+++ b/Engine/Plugins/Enterprise/DatasmithImporter/Source/ExternalSource/Public/IUriResolver.h
@@ -31,5 +31,12 @@ namespace UE::DatasmithImporter
 		 * Return the URI scheme handled by the UriResolver.
 		 */
 		virtual FName GetScheme() const = 0;
+
+#if WITH_EDITOR
+		/**
+		 * Open a dialog window for browsing ExternalSources associated to this UriResolver
+		 */
+		virtual TSharedPtr<FExternalSource> BrowseExternalSource(const FSourceUri& DefaultUri) const = 0;
+#endif //WITH_EDITOR
 	};
 }
\ No newline at end of file
diff --git a/Engine/Source/Editor/UnrealEd/Private/ObjectTools.cpp b/Engine/Source/Editor/UnrealEd/Private/ObjectTools.cpp
--- a/Engine/Source/Editor/UnrealEd/Private/ObjectTools.cpp
+++ b/Engine/Source/Editor/UnrealEd/Private/ObjectTools.cpp
@@ -3974,18 +3974,16 @@ namespace ObjectTools
 	}
 
 	/**
-	 * Internal helper function to obtain format descriptions and extensions of formats supported by the provided factory
+	 * Internal helper function to obtain format descriptions and extensions of formats list
 	 *
-	 * @param	InFactory			Factory whose formats should be retrieved
+	 * @param	Formats				List of formats who should be retrieved
 	 * @param	out_Descriptions	Array of format descriptions associated with the current factory; should equal the number of extensions
 	 * @param	out_Extensions		Array of format extensions associated with the current factory; should equal the number of descriptions
 	 */
-	void InternalGetFactoryFormatInfo( const UFactory* InFactory, TArray<FString>& out_Descriptions, TArray<FString>& out_Extensions )
+	void InternalGetFormatInfo(const TArray<FString>& Formats, TArray<FString>& out_Descriptions, TArray<FString>& out_Extensions )
 	{
-		check(InFactory);
-
-		// Iterate over each format the factory accepts
-		for ( TArray<FString>::TConstIterator FormatIter( InFactory->Formats ); FormatIter; ++FormatIter )
+		// Iterate over each formats.
+		for ( TArray<FString>::TConstIterator FormatIter( Formats ); FormatIter; ++FormatIter )
 		{
 			const FString& CurFormat = *FormatIter;
 
@@ -4038,7 +4036,7 @@ namespace ObjectTools
 
 			TArray<FString> Descriptions;
 			TArray<FString> Extensions;
-			InternalGetFactoryFormatInfo( CurFactory, Descriptions, Extensions );
+			InternalGetFormatInfo( CurFactory->Formats, Descriptions, Extensions );
 			check( Descriptions.Num() == Extensions.Num() );
 
 			// Make sure to only store each key, value pair once
@@ -4118,20 +4116,14 @@ namespace ObjectTools
 		}
 	}
 
-	/**
-	 * Generates a list of file types for a given class.
-	 */
-	void AppendFactoryFileExtensions ( UFactory* InFactory, FString& out_Filetypes, FString& out_Extensions )
+	void InternalAppendFileExtensions(const TArray<FString>& InDescriptions, const TArray<FString>& InExtensions, FString& out_Filetypes, FString& out_Extensions)
 	{
-		TArray<FString> Descriptions;
-		TArray<FString> Extensions;
-		InternalGetFactoryFormatInfo( InFactory, Descriptions, Extensions );
-		check( Descriptions.Num() == Extensions.Num() );
+		check(InDescriptions.Num() == InExtensions.Num());
 
-		for ( int32 FormatIndex = 0; FormatIndex < Descriptions.Num() && FormatIndex < Extensions.Num(); ++FormatIndex )
+		for (int32 FormatIndex = 0; FormatIndex < InDescriptions.Num() && FormatIndex < InExtensions.Num(); ++FormatIndex)
 		{
-			const FString& CurDescription = Descriptions[FormatIndex];
-			const FString& CurExtension = Extensions[FormatIndex];
+			const FString& CurDescription = InDescriptions[FormatIndex];
+			const FString& CurExtension = InExtensions[FormatIndex];
 			const FString& CurLine = FString::Printf( TEXT("%s (*.%s)|*.%s"), *CurDescription, *CurExtension, *CurExtension );
 
 			// Only append the extension if it's not already one of the found extensions
@@ -4156,6 +4148,26 @@ namespace ObjectTools
 		}
 	}
 
+	void AppendFormatsFileExtensions(const TArray<FString>& InFormats, FString& out_FileTypes, FString& out_Extensions)
+	{
+		TArray<FString> Descriptions;
+		TArray<FString> Extensions;
+		InternalGetFormatInfo(InFormats, Descriptions, Extensions);
+		InternalAppendFileExtensions(Descriptions, Extensions, out_FileTypes, out_Extensions);
+	}
+
+	/**
+	 * Generates a list of file types for a given class.
+	 */
+	void AppendFactoryFileExtensions ( UFactory* InFactory, FString& out_Filetypes, FString& out_Extensions )
+	{
+		check(InFactory);
+		TArray<FString> Descriptions;
+		TArray<FString> Extensions;
+		InternalGetFormatInfo( InFactory->Formats, Descriptions, Extensions );
+		InternalAppendFileExtensions( Descriptions, Extensions, out_Filetypes, out_Extensions);
+	}
+
 	/**
 	 * Iterates over all classes and assembles a list of non-abstract UExport-derived type instances.
 	 */
diff --git a/Engine/Source/Editor/UnrealEd/Public/ObjectTools.h b/Engine/Source/Editor/UnrealEd/Public/ObjectTools.h
--- a/Engine/Source/Editor/UnrealEd/Public/ObjectTools.h
+++ b/Engine/Source/Editor/UnrealEd/Public/ObjectTools.h
@@ -465,6 +465,14 @@ namespace ObjectTools
 	 */
 	UNREALED_API void AppendFactoryFileExtensions( UFactory* InFactory, FString& out_Filetypes, FString& out_Extensions );
 
+	/**
+	 * Populates two strings with all of the file types and extensions the format list provides.
+	 *
+	 * @param	InFormats		Array of supported file types. Each entry needs to be of the form "ext;Description" where ext is the file extension. 
+	 * @param	out_FileTypes	File types supported by the provided array of formats, concatenated into a string
+	 * @param	out_Extensions	Extensions supported by the provided array of formats, concatenated into a string
+	 */
+	UNREALED_API void AppendFormatsFileExtensions(const TArray<FString>& InFormats, FString& out_FileTypes, FString& out_Extensions);
 
 	/**
 	 * Iterates over all classes and assembles a list of non-abstract UExport-derived type instances.
