DirectLink in Editor - Minor UI tweaks

diff --git a/Engine/Plugins/Enterprise/DatasmithContent/Source/DatasmithContentEditor/Private/DatasmithImportInfoCustomization.cpp b/Engine/Plugins/Enterprise/DatasmithContent/Source/DatasmithContentEditor/Private/DatasmithImportInfoCustomization.cpp
--- a/Engine/Plugins/Enterprise/DatasmithContent/Source/DatasmithContentEditor/Private/DatasmithImportInfoCustomization.cpp
+++ b/Engine/Plugins/Enterprise/DatasmithContent/Source/DatasmithContentEditor/Private/DatasmithImportInfoCustomization.cpp
@@ -18,7 +18,7 @@
 #include "Widgets/SBoxPanel.h"
 #include "Widgets/Text/STextBlock.h"
 
-#define LOCTEXT_NAMESPACE "AssetImportDataCustomization"
+#define LOCTEXT_NAMESPACE "DatasmithImportInfoCustomization"
 
 FDatasmithImportInfoCustomization::FDatasmithImportInfoCustomization()
 {
@@ -40,7 +40,7 @@ void FDatasmithImportInfoCustomization::CustomizeChildren(TSharedRef<IPropertyHa
 	if (GetImportInfo() != nullptr)
 	{
 		const FSlateFontInfo Font = IDetailLayoutBuilder::GetDetailFont();
-		const FText SourceUriLabel = LOCTEXT("SourceUri", "Source Uri");
+		const FText SourceUriLabel = LOCTEXT("SourceUri", "Source URI");
 
 		const FButtonStyle* ButtonStyle = &FDatasmithContentEditorStyle::Get()->GetWidgetStyle<FButtonStyle>(TEXT("DatasmithDataprepEditor.ButtonLeft"));
 		const FComboBoxStyle* ComboBoxStyle = &FDatasmithContentEditorStyle::Get()->GetWidgetStyle<FComboBoxStyle>(TEXT("DatasmithDataprepEditor.SimpleComboBoxRight"));
@@ -75,7 +75,7 @@ void FDatasmithImportInfoCustomization::CustomizeChildren(TSharedRef<IPropertyHa
 				SNew(SButton)
 				.ButtonStyle(ButtonStyle)
 				.OnClicked(this, &FDatasmithImportInfoCustomization::OnBrowseSourceClicked)
-				.ToolTipText(LOCTEXT("ChangePath_Tooltip", "Browse for a new source file path"))
+				.ToolTipText(LOCTEXT("ChangePath_Tooltip", "Browse for a new source."))
 				[
 					SNew(STextBlock)
 					.Text(LOCTEXT("...", "..."))
