DirectLink in Editor - Replacing temporary icons with proper icons.

diff --git a/Engine/Plugins/Enterprise/DatasmithContent/Source/DatasmithContentEditor/Private/AssetTypeActions_DatasmithScene.cpp b/Engine/Plugins/Enterprise/DatasmithContent/Source/DatasmithContentEditor/Private/AssetTypeActions_DatasmithScene.cpp
--- a/Engine/Plugins/Enterprise/DatasmithContent/Source/DatasmithContentEditor/Private/AssetTypeActions_DatasmithScene.cpp
+++ b/Engine/Plugins/Enterprise/DatasmithContent/Source/DatasmithContentEditor/Private/AssetTypeActions_DatasmithScene.cpp
@@ -4,6 +4,7 @@
 
 #include "DatasmithAssetImportData.h"
 #include "DatasmithContentEditorModule.h"
+#include "DatasmithContentEditorStyle.h"
 #include "DatasmithScene.h"
 
 #include "Framework/MultiBox/MultiBoxBuilder.h"
@@ -48,22 +49,21 @@ void FAssetTypeActions_DatasmithScene::GetActions(const TArray<UObject*>& InObje
 		}
 	}
 
-	const FText AutoReimportText = bHasAutoReimportEnabled
-		? LOCTEXT("DatasmithScene_ToggleDirectLinkAutoReimport_Disable", "Disable Direct Link auto-reimport")
-		: LOCTEXT("DatasmithScene_ToggleDirectLinkAutoReimport_Enable", "Enable Direct Link auto-reimport");
-	const FText AutoReimportTooltip = bHasAutoReimportEnabled
-		? LOCTEXT("DatasmithScene_ToggleDirectLinkAutoReimportTooltip_Disable", "Disable Direct Link auto-reimport for all selected Datasmith Scenes.")
-		: LOCTEXT("DatasmithScene_ToggleDirectLinkAutoReimportTooltip_Enable", "Enable Direct Link auto-reimport for all selected Datasmith Scenes.");
+	const FText AutoReimportText = LOCTEXT("DatasmithScene_ToggleDirectLinkAutoReimport", "Direct Link Auto-Reimport");
+	const FText AutoReimportTooltip = LOCTEXT("DatasmithScene_ToggleDirectLinkAutoReimportTooltip", "Toggle Direct Link Auto-Reimport for all selected Datasmith Scenes.");
 
+	// #ueent_todo DatasmithContent should not know about DirectLink, this should be made modular, or we should directly depend on DirectLinkExtension module.
 	Section.AddMenuEntry(
 		"DatasmithScene_ToggleDirectLinkAutoReimport",
 		AutoReimportText,
 		AutoReimportTooltip,
-		FSlateIcon(FAppStyle::GetAppStyleSetName(), "ContentBrowser.AssetActions.ReimportWithNewFile"),
+		FSlateIcon(FDatasmithContentEditorStyle::GetStyleSetName(), "DatasmithContent.AutoReimportGrayscale"),
 		FUIAction(
 			FExecuteAction::CreateLambda([Scenes = MoveTemp(Scenes), bHasAutoReimportEnabled]() { FAssetTypeActions_DatasmithScene::ExecuteToggleDirectLinkAutoReimport(Scenes, !bHasAutoReimportEnabled); }),
-			FCanExecuteAction::CreateLambda([bHasCanExecuteAutoReimport]() { return bHasCanExecuteAutoReimport; })
+			FCanExecuteAction::CreateLambda([bHasCanExecuteAutoReimport]() { return bHasCanExecuteAutoReimport; }),
+			FGetActionCheckState::CreateLambda([bHasAutoReimportEnabled]() { return bHasAutoReimportEnabled ? ECheckBoxState::Checked : ECheckBoxState::Unchecked; })
 		)
+		, EUserInterfaceActionType::Check
 	);
 }
 
diff --git a/Engine/Plugins/Enterprise/DatasmithContent/Source/DatasmithContentEditor/Private/DatasmithContentEditorStyle.cpp b/Engine/Plugins/Enterprise/DatasmithContent/Source/DatasmithContentEditor/Private/DatasmithContentEditorStyle.cpp
--- a/Engine/Plugins/Enterprise/DatasmithContent/Source/DatasmithContentEditor/Private/DatasmithContentEditorStyle.cpp
+++ b/Engine/Plugins/Enterprise/DatasmithContent/Source/DatasmithContentEditor/Private/DatasmithContentEditorStyle.cpp
@@ -11,7 +11,8 @@
 #include "Styling/SlateTypes.h"
 #include "Styling/StyleColors.h"
 
-#define IMAGE_PLUGIN_BRUSH( RelativePath, ... ) FSlateImageBrush( FDatasmithContentEditorStyle::InContent( RelativePath, ".png" ), __VA_ARGS__ )
+#define IMAGE_PLUGIN_BRUSH( RelativePath, ... ) FSlateImageBrush( FDatasmithContentEditorStyle::InContent( RelativePath, TEXT(".png") ), __VA_ARGS__ )
+#define IMAGE_PLUGIN_BRUSH_SVG(RelativePath, ...) FSlateVectorImageBrush(FDatasmithContentEditorStyle::InContent(RelativePath, TEXT(".svg")), __VA_ARGS__)
 
 TSharedPtr<FSlateStyleSet> FDatasmithContentEditorStyle::StyleSet;
 
@@ -29,47 +30,7 @@ void FDatasmithContentEditorStyle::Initialize()
 	StyleSet->SetContentRoot(FPaths::EngineContentDir() / TEXT("Editor/Slate"));
 	StyleSet->SetCoreContentRoot(FPaths::EngineContentDir() / TEXT("Slate"));
 
-	StyleSet->Set("DatasmithDataprepEditor.Importer", new IMAGE_PLUGIN_BRUSH("Icons/DatasmithImporterIcon40", Icon40x40));
-	StyleSet->Set("DatasmithDataprepEditor.Importer.Small", new IMAGE_PLUGIN_BRUSH("Icons/DatasmithImporterIcon40", Icon20x20));
-	StyleSet->Set("DatasmithDataprepEditor.Importer.Selected", new IMAGE_PLUGIN_BRUSH("Icons/DatasmithImporterIcon40", Icon40x40));
-	StyleSet->Set("DatasmithDataprepEditor.Importer.Selected.Small", new IMAGE_PLUGIN_BRUSH("Icons/DatasmithImporterIcon40", Icon20x20));
-
-	StyleSet->Set("DatasmithDataprepEditor.CADImporter", new IMAGE_PLUGIN_BRUSH("Icons/DatasmithCADImporterIcon40", Icon40x40));
-	StyleSet->Set("DatasmithDataprepEditor.CADImporter.Small", new IMAGE_PLUGIN_BRUSH("Icons/DatasmithCADImporterIcon40", Icon20x20));
-	StyleSet->Set("DatasmithDataprepEditor.CADImporter.Selected", new IMAGE_PLUGIN_BRUSH("Icons/DatasmithCADImporterIcon40", Icon40x40));
-	StyleSet->Set("DatasmithDataprepEditor.CADImporter.Selected.Small", new IMAGE_PLUGIN_BRUSH("Icons/DatasmithCADImporterIcon40", Icon20x20));
-
-	StyleSet->Set("DatasmithDataprepEditor.VREDImporter", new IMAGE_PLUGIN_BRUSH("Icons/DatasmithVREDImporter40", Icon40x40));
-	StyleSet->Set("DatasmithDataprepEditor.VREDImporter.Small", new IMAGE_PLUGIN_BRUSH("Icons/DatasmithVREDImporter40", Icon20x20));
-	StyleSet->Set("DatasmithDataprepEditor.VREDImporter.Selected", new IMAGE_PLUGIN_BRUSH("Icons/DatasmithVREDImporter40", Icon40x40));
-	StyleSet->Set("DatasmithDataprepEditor.VREDImporter.Selected.Small", new IMAGE_PLUGIN_BRUSH("Icons/DatasmithVREDImporter40", Icon20x20));
-
-	StyleSet->Set("DatasmithDataprepEditor.DeltaGenImporter", new IMAGE_PLUGIN_BRUSH("Icons/DatasmithDeltaGenImporter40", Icon40x40));
-	StyleSet->Set("DatasmithDataprepEditor.DeltaGenImporter.Small", new IMAGE_PLUGIN_BRUSH("Icons/DatasmithDeltaGenImporter40", Icon20x20));
-	StyleSet->Set("DatasmithDataprepEditor.DeltaGenImporter.Selected", new IMAGE_PLUGIN_BRUSH("Icons/DatasmithDeltaGenImporter40", Icon40x40));
-	StyleSet->Set("DatasmithDataprepEditor.DeltaGenImporter.Selected.Small", new IMAGE_PLUGIN_BRUSH("Icons/DatasmithDeltaGenImporter40", Icon20x20));
-
-	StyleSet->Set("DatasmithDataprepEditor.SaveScene", new IMAGE_PLUGIN_BRUSH("Icons/SaveScene", Icon40x40));
-	StyleSet->Set("DatasmithDataprepEditor.SaveScene.Small", new IMAGE_PLUGIN_BRUSH("Icons/SaveScene", Icon20x20));
-	StyleSet->Set("DatasmithDataprepEditor.SaveScene.Selected", new IMAGE_PLUGIN_BRUSH("Icons/SaveScene", Icon40x40));
-	StyleSet->Set("DatasmithDataprepEditor.SaveScene.Selected.Small", new IMAGE_PLUGIN_BRUSH("Icons/SaveScene", Icon20x20));
-
-	StyleSet->Set("DatasmithDataprepEditor.ShowDatasmithSceneSettings", new IMAGE_PLUGIN_BRUSH("Icons/IconOptions", Icon40x40));
-
-	StyleSet->Set("DatasmithDataprepEditor.BuildWorld", new IMAGE_PLUGIN_BRUSH("Icons/BuildWorld", Icon40x40));
-	StyleSet->Set("DatasmithDataprepEditor.BuildWorld.Small", new IMAGE_PLUGIN_BRUSH("Icons/BuildWorld", Icon20x20));
-	StyleSet->Set("DatasmithDataprepEditor.BuildWorld.Selected", new IMAGE_PLUGIN_BRUSH("Icons/BuildWorld", Icon40x40));
-	StyleSet->Set("DatasmithDataprepEditor.BuildWorld.Selected.Small", new IMAGE_PLUGIN_BRUSH("Icons/BuildWorld", Icon20x20));
-
-	StyleSet->Set("DatasmithDataprepEditor.ExecutePipeline", new IMAGE_PLUGIN_BRUSH("Icons/ExecutePipeline", Icon40x40));
-	StyleSet->Set("DatasmithDataprepEditor.ExecutePipeline.Small", new IMAGE_PLUGIN_BRUSH("Icons/ExecutePipeline", Icon20x20));
-	StyleSet->Set("DatasmithDataprepEditor.ExecutePipeline.Selected", new IMAGE_PLUGIN_BRUSH("Icons/ExecutePipeline", Icon40x40));
-	StyleSet->Set("DatasmithDataprepEditor.ExecutePipeline.Selected.Small", new IMAGE_PLUGIN_BRUSH("Icons/ExecutePipeline", Icon20x20));
-
-	StyleSet->Set("DatasmithDataprepEditor.Jacketing", new IMAGE_PLUGIN_BRUSH("Icons/Jacketing", Icon40x40));
-	StyleSet->Set("DatasmithDataprepEditor.Jacketing.Small", new IMAGE_PLUGIN_BRUSH("Icons/Jacketing", Icon20x20));
-	StyleSet->Set("DatasmithDataprepEditor.Jacketing.Selected", new IMAGE_PLUGIN_BRUSH("Icons/Jacketing", Icon40x40));
-	StyleSet->Set("DatasmithDataprepEditor.Jacketing.Selected.Small", new IMAGE_PLUGIN_BRUSH("Icons/Jacketing", Icon20x20));
+	StyleSet->Set("DatasmithContent.AutoReimportGrayscale", new IMAGE_PLUGIN_BRUSH_SVG(TEXT("Icons/DataSmithAutoReimport_16"), Icon16x16));
 
 	// Copy the base style, and change the rounding of the left-side borders so that it can be combined with a "right" widget.
 	FButtonStyle ButtonLeftStyle = FAppStyle::Get().GetWidgetStyle<FButtonStyle>(TEXT("Button"));
@@ -112,8 +73,11 @@ FName FDatasmithContentEditorStyle::GetStyleSetName()
 	return StyleName;
 }
 
-FString FDatasmithContentEditorStyle::InContent(const FString& RelativePath, const ANSICHAR* Extension)
+FString FDatasmithContentEditorStyle::InContent(const FString& RelativePath, const TCHAR* Extension)
 {
-	static FString BaseDir = IPluginManager::Get().FindPlugin(DATASMITHCONTENT_MODULE_NAME)->GetBaseDir() + TEXT("/Resources");
-	return (BaseDir / RelativePath) + Extension;
+	static FString ContentDir = IPluginManager::Get().FindPlugin(DATASMITHCONTENT_MODULE_NAME)->GetContentDir();
+	return (ContentDir / RelativePath) + Extension;
 }
+
+#undef IMAGE_PLUGIN_BRUSH_SVG
+#undef IMAGE_PLUGIN_BRUSH
diff --git a/Engine/Plugins/Enterprise/DatasmithContent/Source/DatasmithContentEditor/Public/DatasmithContentEditorStyle.h b/Engine/Plugins/Enterprise/DatasmithContent/Source/DatasmithContentEditor/Public/DatasmithContentEditorStyle.h
--- a/Engine/Plugins/Enterprise/DatasmithContent/Source/DatasmithContentEditor/Public/DatasmithContentEditorStyle.h
+++ b/Engine/Plugins/Enterprise/DatasmithContent/Source/DatasmithContentEditor/Public/DatasmithContentEditorStyle.h
@@ -16,7 +16,7 @@ public:
 	static FName GetStyleSetName();
 
 private:
-	static FString InContent(const FString& RelativePath, const ANSICHAR* Extension);
+	static FString InContent(const FString& RelativePath, const TCHAR* Extension);
 
 private:
 	static TSharedPtr<FSlateStyleSet> StyleSet;
diff --git a/Engine/Plugins/Enterprise/DatasmithImporter/Source/DirectLinkExtension/Private/UI/DirectLinkExtensionStyle.cpp b/Engine/Plugins/Enterprise/DatasmithImporter/Source/DirectLinkExtension/Private/UI/DirectLinkExtensionStyle.cpp
--- a/Engine/Plugins/Enterprise/DatasmithImporter/Source/DirectLinkExtension/Private/UI/DirectLinkExtensionStyle.cpp
+++ b/Engine/Plugins/Enterprise/DatasmithImporter/Source/DirectLinkExtension/Private/UI/DirectLinkExtensionStyle.cpp
@@ -8,6 +8,7 @@
 #include "Interfaces/IPluginManager.h"
 
 #define IMAGE_PLUGIN_BRUSH( RelativePath, ... ) FSlateImageBrush( FDirectLinkExtensionStyle::InContent( RelativePath, TEXT(".png") ), __VA_ARGS__ )
+#define IMAGE_PLUGIN_BRUSH_SVG(RelativePath, ...) FSlateVectorImageBrush(FDirectLinkExtensionStyle::InContent(RelativePath, TEXT(".svg")), __VA_ARGS__)
 
 namespace UE::DatasmithImporter
 {
@@ -36,19 +37,15 @@ namespace UE::DatasmithImporter
 		StyleSet->SetContentRoot(FPaths::EngineContentDir() / TEXT("Editor/Slate"));
 		StyleSet->SetCoreContentRoot(FPaths::EngineContentDir() / TEXT("Slate"));
 
-
-		TUniquePtr< FSlateStyleSet > Style = MakeUnique<FSlateStyleSet>(GetStyleSetName());
-		const FString ContentDir = IPluginManager::Get().FindPlugin(TEXT("DatasmithImporter"))->GetContentDir();
-		Style->SetContentRoot(ContentDir);
-
 		// Const icon sizes
 		const FVector2D Icon16x16(16.0f, 16.0f);
 
 		// 16x16
-		StyleSet->Set("DirectLinkExtension.NotAvailable", new IMAGE_PLUGIN_BRUSH(TEXT("Icons/SourceNotAvailable16"), Icon16x16));
-		StyleSet->Set("DirectLinkExtension.OutOfSync", new IMAGE_PLUGIN_BRUSH(TEXT("Icons/SourceOutOfSync16"), Icon16x16));
-		StyleSet->Set("DirectLinkExtension.UpToDate", new IMAGE_PLUGIN_BRUSH(TEXT("Icons/SourceUpToDate16"), Icon16x16));
-		StyleSet->Set("DirectLinkExtension.AutoReimport", new IMAGE_PLUGIN_BRUSH(TEXT("Icons/SourceAutoReimport16"), Icon16x16));
+		StyleSet->Set("DirectLinkExtension.NotAvailable", new IMAGE_PLUGIN_BRUSH_SVG(TEXT("Icons/DirectLinkBadgeNotAvailable"), Icon16x16));
+		StyleSet->Set("DirectLinkExtension.OutOfSync", new IMAGE_PLUGIN_BRUSH_SVG(TEXT("Icons/DirectLinkBadgeOutOfSync"), Icon16x16));
+		StyleSet->Set("DirectLinkExtension.UpToDate", new IMAGE_PLUGIN_BRUSH_SVG(TEXT("Icons/DirectLinkBadgeUpToDate"), Icon16x16));
+		StyleSet->Set("DirectLinkExtension.AutoReimport", new IMAGE_PLUGIN_BRUSH_SVG(TEXT("Icons/DirectLinkBadgeAutoReimport"), Icon16x16));
+		StyleSet->Set("DirectLinkExtension.AutoReimportGrayscale", new IMAGE_PLUGIN_BRUSH_SVG(TEXT("Icons/DirectLinkAutoReimport_16"), Icon16x16));
 
 		FSlateStyleRegistry::RegisterSlateStyle(*StyleSet.Get());
 	};
@@ -68,4 +65,5 @@ namespace UE::DatasmithImporter
 	}
 }
 
-#undef IMAGE_PLUGIN_BRUSH
\ No newline at end of file
+#undef IMAGE_PLUGIN_BRUSH
+#undef IMAGE_PLUGIN_BRUSH_SVG
\ No newline at end of file
