[Visual Dataprep] Improve workflows for using Metadata

diff --git a/Engine/Plugins/Enterprise/DataprepEditor/Source/DataprepEditor/Private/Widgets/DataprepAssetView.cpp b/Engine/Plugins/Enterprise/DataprepEditor/Source/DataprepEditor/Private/Widgets/DataprepAssetView.cpp
--- a/Engine/Plugins/Enterprise/DataprepEditor/Source/DataprepEditor/Private/Widgets/DataprepAssetView.cpp
+++ b/Engine/Plugins/Enterprise/DataprepEditor/Source/DataprepEditor/Private/Widgets/DataprepAssetView.cpp
@@ -82,7 +82,17 @@ void SGraphNodeDetailsWidget::Construct(const FArguments& InArgs)
 
 	PropertyView = EditModule.CreateDetailView(DetailsViewArgs);
 
-	PropertyView->GetIsPropertyEditingEnabledDelegate().BindSP(this, &SGraphNodeDetailsWidget::GetCanEditProperties);
+	// Enable selection (copy/paste) for metadata property only
+
+	PropertyView->GetIsCustomRowReadOnlyDelegate().BindLambda([this](const FName InRowName, const FName InParentName) -> bool
+		{
+			return !bCanEditProperties;
+		});
+
+	PropertyView->GetIsPropertyReadOnlyDelegate().BindLambda([this](const FPropertyAndParent& PNP) -> bool
+		{
+			return !bCanEditProperties && (PNP.Property.GetName() != TEXT("MetaData"));
+		});
 
 	DetailsSplitter = SNew(SSplitter)
 		.MinimumSlotHeight(80.0f)
diff --git a/Engine/Plugins/Enterprise/DatasmithContent/Source/DatasmithContent/Private/DatasmithAssetUserData.cpp b/Engine/Plugins/Enterprise/DatasmithContent/Source/DatasmithContent/Private/DatasmithAssetUserData.cpp
--- a/Engine/Plugins/Enterprise/DatasmithContent/Source/DatasmithContent/Private/DatasmithAssetUserData.cpp
+++ b/Engine/Plugins/Enterprise/DatasmithContent/Source/DatasmithContent/Private/DatasmithAssetUserData.cpp
@@ -52,14 +52,36 @@ UDatasmithAssetUserData* UDatasmithAssetUserData::GetDatasmithUserData(UObject*
 	return nullptr;
 }
 
-FString UDatasmithAssetUserData::GetDatasmithUserDataValueForKey(UObject* Object, FName Key)
+FString UDatasmithAssetUserData::GetDatasmithUserDataValueForKey(UObject* Object, FName Key, bool bPartialMatchKey)
 {
 	if (Object)
 	{
 		if (UDatasmithAssetUserData* AssetUserData = GetDatasmithUserData(Object))
 		{
-			FString* Value = AssetUserData->MetaData.Find(Key);
-			return Value ? *Value : FString();
+			FString Value;
+
+			if (bPartialMatchKey)
+			{
+				const FString KeyString = Key.ToString();
+
+				for (TPair<FName, FString> KeyValuePair : AssetUserData->MetaData)
+				{
+					if (KeyValuePair.Key.ToString().Contains(KeyString))
+					{
+						Value = KeyValuePair.Value;
+						break;
+					}
+				}
+			}
+			else
+			{
+				if (FString* ValuePtr = AssetUserData->MetaData.Find(Key))
+				{
+					Value = *ValuePtr;
+				}
+			}
+			
+			return Value;
 		}
 	}
 	return FString();
diff --git a/Engine/Plugins/Enterprise/DatasmithContent/Source/DatasmithContent/Private/DatasmithContentBlueprintLibrary.cpp b/Engine/Plugins/Enterprise/DatasmithContent/Source/DatasmithContent/Private/DatasmithContentBlueprintLibrary.cpp
--- a/Engine/Plugins/Enterprise/DatasmithContent/Source/DatasmithContent/Private/DatasmithContentBlueprintLibrary.cpp
+++ b/Engine/Plugins/Enterprise/DatasmithContent/Source/DatasmithContent/Private/DatasmithContentBlueprintLibrary.cpp
@@ -13,9 +13,9 @@ UDatasmithAssetUserData* UDatasmithContentBlueprintLibrary::GetDatasmithUserData
 	return UDatasmithAssetUserData::GetDatasmithUserData(Object);
 }
 
-FString UDatasmithContentBlueprintLibrary::GetDatasmithUserDataValueForKey(UObject* Object, FName Key)
+FString UDatasmithContentBlueprintLibrary::GetDatasmithUserDataValueForKey(UObject* Object, FName Key, bool bPartialMatchKey)
 {
-	return UDatasmithAssetUserData::GetDatasmithUserDataValueForKey(Object, Key);
+	return UDatasmithAssetUserData::GetDatasmithUserDataValueForKey(Object, Key, bPartialMatchKey);
 }
 
 void UDatasmithContentBlueprintLibrary::GetDatasmithUserDataKeysAndValuesForValue(UObject* Object, const FString& StringToMatch, TArray<FName>& OutKeys, TArray<FString>& OutValues)
diff --git a/Engine/Plugins/Enterprise/DatasmithContent/Source/DatasmithContent/Public/DatasmithAssetUserData.h b/Engine/Plugins/Enterprise/DatasmithContent/Source/DatasmithContent/Public/DatasmithAssetUserData.h
--- a/Engine/Plugins/Enterprise/DatasmithContent/Source/DatasmithContent/Public/DatasmithAssetUserData.h
+++ b/Engine/Plugins/Enterprise/DatasmithContent/Source/DatasmithContent/Public/DatasmithAssetUserData.h
@@ -29,7 +29,7 @@ public:
 	virtual void PostLoad() override;
 #endif
 
-	static FString GetDatasmithUserDataValueForKey(UObject* Object, FName Key);
+	static FString GetDatasmithUserDataValueForKey(UObject* Object, FName Key, bool bPartialMatchKey = false);
 	static UDatasmithAssetUserData* GetDatasmithUserData(UObject* Object);
 	static bool SetDatasmithUserDataValueForKey(UObject* Object, FName Key, const FString & Value);
 
diff --git a/Engine/Plugins/Enterprise/DatasmithContent/Source/DatasmithContent/Public/DatasmithContentBlueprintLibrary.h b/Engine/Plugins/Enterprise/DatasmithContent/Source/DatasmithContent/Public/DatasmithContentBlueprintLibrary.h
--- a/Engine/Plugins/Enterprise/DatasmithContent/Source/DatasmithContent/Public/DatasmithContentBlueprintLibrary.h
+++ b/Engine/Plugins/Enterprise/DatasmithContent/Source/DatasmithContent/Public/DatasmithContentBlueprintLibrary.h
@@ -25,12 +25,13 @@ public:
 
 	/**
 	 * Get the value of the given key for the Datasmith User Data of the given object.
-	 * @param	Object	The Object from which to retrieve the Datasmith User Data.
-	 * @param	Key		The key to find in the Datasmith User Data.
-	 * @return			The string value associated with the given key
+	 * @param	Object				The Object from which to retrieve the Datasmith User Data.
+	 * @param	Key					The key to find in the Datasmith User Data.
+	 * @param	bPartialMatchKey	If true, check for contains, rather than exact match.
+	 * @return						The string value associated with the given key
 	 */
 	UFUNCTION(BlueprintCallable, Category = "Datasmith User Data")
-	static FString GetDatasmithUserDataValueForKey(UObject* Object, FName Key);
+	static FString GetDatasmithUserDataValueForKey(UObject* Object, FName Key, bool bPartialMatchKey = false);
 
 	/**
 	 * Get the keys and values for which the associated value contains the string to match for the Datasmith User Data of the given object.
diff --git a/Engine/Plugins/Enterprise/DatasmithImporter/Source/DatasmithImporter/Private/DatasmithDataprepFetcher.cpp b/Engine/Plugins/Enterprise/DatasmithImporter/Source/DatasmithImporter/Private/DatasmithDataprepFetcher.cpp
--- a/Engine/Plugins/Enterprise/DatasmithImporter/Source/DatasmithImporter/Private/DatasmithDataprepFetcher.cpp
+++ b/Engine/Plugins/Enterprise/DatasmithImporter/Source/DatasmithImporter/Private/DatasmithDataprepFetcher.cpp
@@ -13,7 +13,8 @@ FString UDatasmithStringMetadataValueFetcher::Fetch_Implementation(const UObject
 	if ( Object ) 
 	{
 		bOutFetchSucceded = true;
-		return UDatasmithContentBlueprintLibrary::GetDatasmithUserDataValueForKey( const_cast< UObject* >( Object ), Key );
+		bool bMatchPartialKey = KeyMatch == EMetadataKeyMatchingCriteria::Contains;
+		return UDatasmithContentBlueprintLibrary::GetDatasmithUserDataValueForKey( const_cast< UObject* >( Object ), Key, bMatchPartialKey);
 	}
 
 	bOutFetchSucceded = false;
diff --git a/Engine/Plugins/Enterprise/DatasmithImporter/Source/DatasmithImporter/Private/DatasmithDataprepFetcher.h b/Engine/Plugins/Enterprise/DatasmithImporter/Source/DatasmithImporter/Private/DatasmithDataprepFetcher.h
--- a/Engine/Plugins/Enterprise/DatasmithImporter/Source/DatasmithImporter/Private/DatasmithDataprepFetcher.h
+++ b/Engine/Plugins/Enterprise/DatasmithImporter/Source/DatasmithImporter/Private/DatasmithDataprepFetcher.h
@@ -8,6 +8,13 @@
 
 #include "DatasmithDataprepFetcher.generated.h"
 
+UENUM()
+enum EMetadataKeyMatchingCriteria
+{
+	ExactMatch,
+	Contains
+};
+
 UCLASS(BlueprintType, NotBlueprintable, Meta = (DisplayName="Metadata Value", ToolTip="Filter objects based on the key value of their metadata."))
 class UDatasmithStringMetadataValueFetcher final : public UDataprepStringFetcher
 {
@@ -22,6 +29,8 @@ public:
 	virtual bool IsThreadSafe() const final;
 	//~ End of UDataprepFetcher interface
 
+	UPROPERTY(EditInstanceOnly, BlueprintReadWrite, Category = "Key")
+	TEnumAsByte<EMetadataKeyMatchingCriteria> KeyMatch;
 
 	// The key for the for the string
 	UPROPERTY(EditInstanceOnly, BlueprintReadWrite, Category = "Key")
