Limits are defined for tessellation parameters:
- 0.03 for the chord error
- 5 degrees for the max angle
- Edge length could not be limits as 0 = no Edge length. Its value, if != 0 is limited to 0.5 (with a warning message)
What ever, if the user, for a specific need, want to set smaller value for metric parameters, the CVar ds.CADTranslator.MeshingParameterFactor can be used

diff --git a/Engine/Plugins/Enterprise/DatasmithCADImporter/Source/CADTools/Private/CADOption.cpp b/Engine/Plugins/Enterprise/DatasmithCADImporter/Source/CADTools/Private/CADOption.cpp
--- a/Engine/Plugins/Enterprise/DatasmithCADImporter/Source/CADTools/Private/CADOption.cpp
+++ b/Engine/Plugins/Enterprise/DatasmithCADImporter/Source/CADTools/Private/CADOption.cpp
@@ -79,6 +79,16 @@ This welding is performed respecting the ds.CADTranslator.StitchingTolerance\n\
 Default value is true\n"),
 ECVF_Default);
 
+float FImportParameters::GMeshingParameterFactor = 1;
+FAutoConsoleVariableRef GCADTranslatorMeshingParameterFactor(
+	TEXT("ds.CADTranslator.MeshingParameterFactor"),
+	FImportParameters::GMeshingParameterFactor,
+	TEXT("Factor to allow to use smaller value than the defined minimal value of metric meshing parameters (i.e. Chord error > 0.005 cm, Max Edge Length > 1. cm) \n\
+The used value of the meshing parameter is value * MeshingParameterFactor \n\
+Default value is 1.\n"),
+ECVF_Default);
+
+
 static bool bGAliasSewByColor = false;
 FAutoConsoleVariableRef GAliasSewByColor(
 	TEXT("ds.CADTranslator.Alias.SewByColor"),
diff --git a/Engine/Plugins/Enterprise/DatasmithCADImporter/Source/CADTools/Public/CADOptions.h b/Engine/Plugins/Enterprise/DatasmithCADImporter/Source/CADTools/Public/CADOptions.h
--- a/Engine/Plugins/Enterprise/DatasmithCADImporter/Source/CADTools/Public/CADOptions.h
+++ b/Engine/Plugins/Enterprise/DatasmithCADImporter/Source/CADTools/Public/CADOptions.h
@@ -37,6 +37,7 @@ namespace CADLibrary
 		CADTOOLS_API static bool bGSewMeshIfNeeded;
 		CADTOOLS_API static float GStitchingTolerance;
 		CADTOOLS_API static float GUnitScale;
+		CADTOOLS_API static float GMeshingParameterFactor;
 
 	public:
 		FImportParameters(FDatasmithUtils::EModelCoordSystem NewCoordinateSystem = FDatasmithUtils::EModelCoordSystem::ZUp_RightHanded)
@@ -50,8 +51,8 @@ namespace CADLibrary
 	
 		void SetTesselationParameters(double InChordTolerance, double InMaxEdgeLength, double InMaxNormalAngle, CADLibrary::EStitchingTechnique InStitchingTechnique)
 		{
-			ChordTolerance = InChordTolerance;
-			MaxEdgeLength = InMaxEdgeLength;
+			ChordTolerance = InChordTolerance * GMeshingParameterFactor;
+			MaxEdgeLength = InMaxEdgeLength * GMeshingParameterFactor;
 			MaxNormalAngle = InMaxNormalAngle;
 			StitchingTechnique = InStitchingTechnique;
 		}
diff --git a/Engine/Plugins/Enterprise/DatasmithCADImporter/Source/DatasmithCADTranslator/Private/DatasmithCADTranslator.cpp b/Engine/Plugins/Enterprise/DatasmithCADImporter/Source/DatasmithCADTranslator/Private/DatasmithCADTranslator.cpp
--- a/Engine/Plugins/Enterprise/DatasmithCADImporter/Source/DatasmithCADTranslator/Private/DatasmithCADTranslator.cpp
+++ b/Engine/Plugins/Enterprise/DatasmithCADImporter/Source/DatasmithCADTranslator/Private/DatasmithCADTranslator.cpp
@@ -141,7 +141,20 @@ bool FDatasmithCADTranslator::LoadScene(TSharedRef<IDatasmithScene> DatasmithSce
 	UE_LOG(LogCADTranslator, Display, TEXT(" - Processing:           %s")
 		, CADLibrary::FImportParameters::bGEnableCADCache ? (CADLibrary::GMaxImportThreads == 1 ? TEXT("Sequencial") : TEXT("Parallel")) : TEXT("Sequencial"));
 
-	ImportParameters.SetTesselationParameters(TesselationOptions.ChordTolerance, TesselationOptions.MaxEdgeLength, TesselationOptions.NormalTolerance, (CADLibrary::EStitchingTechnique)TesselationOptions.StitchingTechnique);
+	TFunction<double(double, double, const TCHAR*)> CheckParameterValue = [](double Value, double MinValue, const TCHAR* ParameterName) -> double
+	{
+		if (Value < MinValue)
+		{
+			UE_LOG(LogCADTranslator, Warning, TEXT("%s value (%f) of tessellation parameters is smaller than the minimal value %d. It's value is modified to respect the limit"), ParameterName, Value, MinValue);
+			return MinValue;
+		}
+		return Value;
+	};
+	
+	ImportParameters.SetTesselationParameters(CheckParameterValue(TesselationOptions.ChordTolerance, UE::DatasmithTessellation::MinTessellationChord, TEXT("Chord tolerance")),
+		FMath::IsNearlyZero(TesselationOptions.MaxEdgeLength) ? 0. : CheckParameterValue(TesselationOptions.MaxEdgeLength, UE::DatasmithTessellation::MinTessellationEdgeLength, TEXT("Max Edge Length")),
+		CheckParameterValue(TesselationOptions.NormalTolerance, UE::DatasmithTessellation::MinTessellationAngle, TEXT("Max Angle")),
+		(CADLibrary::EStitchingTechnique)TesselationOptions.StitchingTechnique);
 	ImportParameters.SetModelCoordinateSystem(FDatasmithUtils::EModelCoordSystem::ZUp_RightHanded);
 
 	switch (FileDescriptor.GetFileFormat())
diff --git a/Engine/Plugins/Enterprise/DatasmithContent/Source/DatasmithContent/Public/DatasmithImportOptions.h b/Engine/Plugins/Enterprise/DatasmithContent/Source/DatasmithContent/Public/DatasmithImportOptions.h
--- a/Engine/Plugins/Enterprise/DatasmithContent/Source/DatasmithContent/Public/DatasmithImportOptions.h
+++ b/Engine/Plugins/Enterprise/DatasmithContent/Source/DatasmithContent/Public/DatasmithImportOptions.h
@@ -211,6 +211,13 @@ struct DATASMITHCONTENT_API FDatasmithImportBaseOptions
 	bool CanIncludeAnimation() const { return bIncludeGeometry || bIncludeCamera || bIncludeLight; }
 };
 
+namespace UE::DatasmithTessellation
+{
+const double MinTessellationAngle = 5;
+const double MinTessellationEdgeLength = 1.;
+const double MinTessellationChord = 0.005;  // Usual value in CAD software is 0.02 cm
+}
+
 USTRUCT(BlueprintType)
 struct DATASMITHCONTENT_API FDatasmithTessellationOptions
 {
@@ -227,16 +234,16 @@ struct DATASMITHCONTENT_API FDatasmithTessellationOptions
 	/**
 	 * Maximum distance between any point on a triangle generated by the tessellation process and the actual surface.
 	 * The lower the value the more triangles.
-	 * Default value is 0.2.
+	 * Default value is 0.2, minimal value is 0.005 cm.
 	 */
-	UPROPERTY(config, EditAnywhere, BlueprintReadWrite, Category = "Geometry & Tessellation Options", meta = (Units = cm, ToolTip = "Maximum distance between any generated triangle and the original surface. Smaller values make more triangles.", ClampMin = "0.0"))
+	UPROPERTY(config, EditAnywhere, BlueprintReadWrite, Category = "Geometry & Tessellation Options", meta = (Units = cm, ToolTip = "Maximum distance between any generated triangle and the original surface. Smaller values make more triangles.", ClampMin = "0.005"))
 	float ChordTolerance;
 
 	/**
 	 * Maximum length of edges of triangles generated by the tessellation process.
 	 * The length is in scene/model unit. The smaller the more triangles are generated.
 	 * Value of 0 means no constraint on length of edges
-	 * Default value is 0.
+	 * Default value is 0 to disable this criteria, and 1. cm is its minimal value if enable.
 	 */
 	UPROPERTY(config, EditAnywhere, BlueprintReadWrite, AdvancedDisplay, Category = "Geometry & Tessellation Options", meta = (Units = cm, DisplayName = "Max Edge Length", ToolTip = "Maximum length of any edge in the generated triangles. Smaller values make more triangles.", ClampMin = "0.0"))
 	float MaxEdgeLength;
@@ -244,9 +251,9 @@ struct DATASMITHCONTENT_API FDatasmithTessellationOptions
 	/**
 	 * Maximum angle between the normal of two triangles generated by the tessellation process.
 	 * The angle is expressed in degree. The smaller the more triangles are generated.
-	 * Default value is 20 degrees.
+	 * Default value is 20 degrees, min value is 5 degrees.
 	 */
-	UPROPERTY(config, EditAnywhere, BlueprintReadWrite, AdvancedDisplay, Category = "Geometry & Tessellation Options", meta = (Units = deg, ToolTip = "Maximum angle between adjacent triangles. Smaller values make more triangles.", ClampMin = "0.0", ClampMax = "90.0"))
+	UPROPERTY(config, EditAnywhere, BlueprintReadWrite, AdvancedDisplay, Category = "Geometry & Tessellation Options", meta = (Units = deg, ToolTip = "Maximum angle between adjacent triangles. Smaller values make more triangles.", ClampMin = "5.0", ClampMax = "90.0"))
 	float NormalTolerance;
 
 
