Dataprep: Fixing Metadata filter

diff --git a/Engine/Plugins/Enterprise/DatasmithContent/Source/DatasmithContent/Private/DatasmithAssetUserData.cpp b/Engine/Plugins/Enterprise/DatasmithContent/Source/DatasmithContent/Private/DatasmithAssetUserData.cpp
--- a/Engine/Plugins/Enterprise/DatasmithContent/Source/DatasmithContent/Private/DatasmithAssetUserData.cpp
+++ b/Engine/Plugins/Enterprise/DatasmithContent/Source/DatasmithContent/Private/DatasmithAssetUserData.cpp
@@ -87,6 +87,40 @@ FString UDatasmithAssetUserData::GetDatasmithUserDataValueForKey(UObject* Object
 	return FString();
 }
 
+TArray<FString> UDatasmithAssetUserData::GetDatasmithUserDataValuesForKey(UObject* Object, FName Key, bool bPartialMatchKey)
+{
+	if (Object)
+	{
+		if (UDatasmithAssetUserData* AssetUserData = GetDatasmithUserData(Object))
+		{
+			TArray<FString> Values;
+
+			if (bPartialMatchKey)
+			{
+				const FString KeyString = Key.ToString();
+
+				for (TPair<FName, FString> KeyValuePair : AssetUserData->MetaData)
+				{
+					if (KeyValuePair.Key.ToString().Contains(KeyString))
+					{
+						Values.Add(KeyValuePair.Value);
+					}
+				}
+			}
+			else
+			{
+				if (FString* ValuePtr = AssetUserData->MetaData.Find(Key))
+				{
+					Values.Add(*ValuePtr);
+				}
+			}
+
+			return Values;
+		}
+	}
+	return {};
+}
+
 bool UDatasmithAssetUserData::SetDatasmithUserDataValueForKey(UObject* Object, FName Key, const FString & Value)
 {
 	// For AActor, the interface is actually implemented by the ActorComponent
diff --git a/Engine/Plugins/Enterprise/DatasmithContent/Source/DatasmithContent/Private/DatasmithContentBlueprintLibrary.cpp b/Engine/Plugins/Enterprise/DatasmithContent/Source/DatasmithContent/Private/DatasmithContentBlueprintLibrary.cpp
--- a/Engine/Plugins/Enterprise/DatasmithContent/Source/DatasmithContent/Private/DatasmithContentBlueprintLibrary.cpp
+++ b/Engine/Plugins/Enterprise/DatasmithContent/Source/DatasmithContent/Private/DatasmithContentBlueprintLibrary.cpp
@@ -18,6 +18,11 @@ FString UDatasmithContentBlueprintLibrary::GetDatasmithUserDataValueForKey(UObje
 	return UDatasmithAssetUserData::GetDatasmithUserDataValueForKey(Object, Key, bPartialMatchKey);
 }
 
+TArray<FString> UDatasmithContentBlueprintLibrary::GetDatasmithUserDataValuesForKey(UObject* Object, FName Key, bool bPartialMatchKey)
+{
+	return UDatasmithAssetUserData::GetDatasmithUserDataValuesForKey(Object, Key, bPartialMatchKey);
+}
+
 void UDatasmithContentBlueprintLibrary::GetDatasmithUserDataKeysAndValuesForValue(UObject* Object, const FString& StringToMatch, TArray<FName>& OutKeys, TArray<FString>& OutValues)
 {
 	OutKeys.Reset();
diff --git a/Engine/Plugins/Enterprise/DatasmithContent/Source/DatasmithContent/Public/DatasmithAssetUserData.h b/Engine/Plugins/Enterprise/DatasmithContent/Source/DatasmithContent/Public/DatasmithAssetUserData.h
--- a/Engine/Plugins/Enterprise/DatasmithContent/Source/DatasmithContent/Public/DatasmithAssetUserData.h
+++ b/Engine/Plugins/Enterprise/DatasmithContent/Source/DatasmithContent/Public/DatasmithAssetUserData.h
@@ -30,6 +30,7 @@ public:
 #endif
 
 	static FString GetDatasmithUserDataValueForKey(UObject* Object, FName Key, bool bPartialMatchKey = false);
+	static TArray<FString> GetDatasmithUserDataValuesForKey(UObject* Object, FName Key, bool bPartialMatchKey = false);
 	static UDatasmithAssetUserData* GetDatasmithUserData(UObject* Object);
 	static bool SetDatasmithUserDataValueForKey(UObject* Object, FName Key, const FString & Value);
 
diff --git a/Engine/Plugins/Enterprise/DatasmithContent/Source/DatasmithContent/Public/DatasmithContentBlueprintLibrary.h b/Engine/Plugins/Enterprise/DatasmithContent/Source/DatasmithContent/Public/DatasmithContentBlueprintLibrary.h
--- a/Engine/Plugins/Enterprise/DatasmithContent/Source/DatasmithContent/Public/DatasmithContentBlueprintLibrary.h
+++ b/Engine/Plugins/Enterprise/DatasmithContent/Source/DatasmithContent/Public/DatasmithContentBlueprintLibrary.h
@@ -27,12 +27,21 @@ public:
 	 * Get the value of the given key for the Datasmith User Data of the given object.
 	 * @param	Object				The Object from which to retrieve the Datasmith User Data.
 	 * @param	Key					The key to find in the Datasmith User Data.
-	 * @param	bPartialMatchKey	If true, check for contains, rather than exact match.
 	 * @return						The string value associated with the given key
 	 */
 	UFUNCTION(BlueprintCallable, Category = "Datasmith User Data")
 	static FString GetDatasmithUserDataValueForKey(UObject* Object, FName Key, bool bPartialMatchKey = false);
 
+	/**
+	 * Get the values of the given key for the Datasmith User Data of the given object.
+	 * @param	Object				The Object from which to retrieve the Datasmith User Data.
+	 * @param	Key					The key to find in the Datasmith User Data.
+	 * @param	bPartialMatchKey	If true, check for contains, rather than exact match.
+	 * @return						The string value associated with the given key
+	 */
+	UFUNCTION(BlueprintCallable, Category = "Datasmith User Data")
+	static TArray<FString> GetDatasmithUserDataValuesForKey(UObject* Object, FName Key, bool bPartialMatchKey = false);
+
 	/**
 	 * Get the keys and values for which the associated value contains the string to match for the Datasmith User Data of the given object.
 	 * @param	Object			The Object from which to retrieve the Datasmith User Data.
diff --git a/Engine/Plugins/Enterprise/DatasmithImporter/Source/DatasmithImporter/Private/DatasmithDataprepFetcher.cpp b/Engine/Plugins/Enterprise/DatasmithImporter/Source/DatasmithImporter/Private/DatasmithDataprepFetcher.cpp
--- a/Engine/Plugins/Enterprise/DatasmithImporter/Source/DatasmithImporter/Private/DatasmithDataprepFetcher.cpp
+++ b/Engine/Plugins/Enterprise/DatasmithImporter/Source/DatasmithImporter/Private/DatasmithDataprepFetcher.cpp
@@ -8,13 +8,13 @@
 
 /* UDataprepStringMetadataValueFetcher methods
  *****************************************************************************/
-FString UDatasmithStringMetadataValueFetcher::Fetch_Implementation(const UObject* Object, bool& bOutFetchSucceded) const
+TArray<FString> UDatasmithStringMetadataValueFetcher::Fetch_Implementation(const UObject* Object, bool& bOutFetchSucceded) const
 {
 	if ( Object ) 
 	{
 		bOutFetchSucceded = true;
 		bool bMatchPartialKey = KeyMatch == EMetadataKeyMatchingCriteria::Contains;
-		return UDatasmithContentBlueprintLibrary::GetDatasmithUserDataValueForKey( const_cast< UObject* >( Object ), Key, bMatchPartialKey);
+		return UDatasmithContentBlueprintLibrary::GetDatasmithUserDataValuesForKey( const_cast< UObject* >( Object ), Key, bMatchPartialKey);
 	}
 
 	bOutFetchSucceded = false;
diff --git a/Engine/Plugins/Enterprise/DatasmithImporter/Source/DatasmithImporter/Private/DatasmithDataprepFetcher.h b/Engine/Plugins/Enterprise/DatasmithImporter/Source/DatasmithImporter/Private/DatasmithDataprepFetcher.h
--- a/Engine/Plugins/Enterprise/DatasmithImporter/Source/DatasmithImporter/Private/DatasmithDataprepFetcher.h
+++ b/Engine/Plugins/Enterprise/DatasmithImporter/Source/DatasmithImporter/Private/DatasmithDataprepFetcher.h
@@ -2,7 +2,7 @@
 
 #pragma once
 
-#include "SelectionSystem/DataprepStringFetcher.h"
+#include "SelectionSystem/DataprepStringsArrayFetcher.h"
 
 #include "CoreMinimal.h"
 
@@ -16,13 +16,13 @@ enum EMetadataKeyMatchingCriteria
 };
 
 UCLASS(BlueprintType, NotBlueprintable, Meta = (DisplayName="Metadata Value", ToolTip="Filter objects based on the key value of their metadata."))
-class UDatasmithStringMetadataValueFetcher final : public UDataprepStringFetcher
+class UDatasmithStringMetadataValueFetcher final : public UDataprepStringsArrayFetcher
 {
 	GENERATED_BODY()
 public:
-	//~ UDataprepStringFetcher interface
-	virtual FString Fetch_Implementation(const UObject* Object, bool& bOutFetchSucceded) const final;
-	//~ End of UDataprepStringFetcher interface
+	//~ UDataprepStringsArrayFetcher interface
+	virtual TArray<FString> Fetch_Implementation(const UObject* Object, bool& bOutFetchSucceded) const override;
+	//~ End of UDataprepStringsArrayFetcher interface
 
 	//~ UDataprepFetcher interface
 	virtual FText GetNodeDisplayFetcherName_Implementation() const;
