Followup to CL 22508580, fixing discovering several implementations that remain using HasActions.  More things to be upgraded.

diff --git a/Engine/Plugins/Editor/ConsoleVariablesEditor/Source/ConsoleVariablesEditor/Private/AssetTypeActions/AssetTypeActions_ConsoleVariables.h b/Engine/Plugins/Editor/ConsoleVariablesEditor/Source/ConsoleVariablesEditor/Private/AssetTypeActions/AssetTypeActions_ConsoleVariables.h
--- a/Engine/Plugins/Editor/ConsoleVariablesEditor/Source/ConsoleVariablesEditor/Private/AssetTypeActions/AssetTypeActions_ConsoleVariables.h
+++ b/Engine/Plugins/Editor/ConsoleVariablesEditor/Source/ConsoleVariablesEditor/Private/AssetTypeActions/AssetTypeActions_ConsoleVariables.h
@@ -33,15 +33,6 @@ public:
 	{
 		return FColor(238, 181, 235, 255);
 	}
-	
-	/**
-	 * @param InObjects A list of objects selected in the content browser
-	 * @return Whether or not the asset has context actions (required to return true in order to use registered actions)
-	 */
-	virtual bool HasActions(const TArray<UObject*>& InObjects) const override
-	{
-		return true;
-	}
 
 	// These are defined in source - GetName requires a defined LOCTEXT namespace and GetSupportedClass has an external include
 	
diff --git a/Engine/Plugins/Enterprise/DataprepEditor/Source/DataprepEditor/Private/AssetTypeActions_DataprepAssetInterface.h b/Engine/Plugins/Enterprise/DataprepEditor/Source/DataprepEditor/Private/AssetTypeActions_DataprepAssetInterface.h
--- a/Engine/Plugins/Enterprise/DataprepEditor/Source/DataprepEditor/Private/AssetTypeActions_DataprepAssetInterface.h
+++ b/Engine/Plugins/Enterprise/DataprepEditor/Source/DataprepEditor/Private/AssetTypeActions_DataprepAssetInterface.h
@@ -16,7 +16,6 @@ public:
 	virtual uint32 GetCategories() override;
 	virtual FColor GetTypeColor() const override { return FColor(0, 0, 0); }
 	virtual UClass* GetSupportedClass() const override;
-	virtual bool HasActions ( const TArray<UObject*>& InObjects ) const override { return true; }
 	virtual void GetActions( const TArray<UObject*>& InObjects, FMenuBuilder& MenuBuilder ) override;
 	virtual bool IsImportedAsset() const override { return false; }
 
diff --git a/Engine/Plugins/Enterprise/DatasmithContent/Source/DatasmithContentEditor/Private/AssetTypeActions_DatasmithScene.h b/Engine/Plugins/Enterprise/DatasmithContent/Source/DatasmithContentEditor/Private/AssetTypeActions_DatasmithScene.h
--- a/Engine/Plugins/Enterprise/DatasmithContent/Source/DatasmithContentEditor/Private/AssetTypeActions_DatasmithScene.h
+++ b/Engine/Plugins/Enterprise/DatasmithContent/Source/DatasmithContentEditor/Private/AssetTypeActions_DatasmithScene.h
@@ -15,7 +15,6 @@ public:
 	virtual FColor GetTypeColor() const override { return FColor(255, 0, 0); }
 	virtual UClass* GetSupportedClass() const override;
 	virtual void GetActions(const TArray<UObject*>& InObjects, FToolMenuSection& Section) override;
-	virtual bool HasActions ( const TArray<UObject*>& InObjects ) const override { return true; }
 	virtual bool IsImportedAsset() const override { return true; }
 	virtual void GetResolvedSourceFilePaths(const TArray<UObject*>& TypeAssets, TArray<FString>& OutSourceFilePaths) const override;
 	virtual void OpenAssetEditor(const TArray<UObject*>& InObjects, TSharedPtr<class IToolkitHost> EditWithinLevelEditor = TSharedPtr<IToolkitHost>()) override;
diff --git a/Engine/Plugins/Enterprise/LidarPointCloud/Source/LidarPointCloudEditor/Private/LidarPointCloudFactory.h b/Engine/Plugins/Enterprise/LidarPointCloud/Source/LidarPointCloudEditor/Private/LidarPointCloudFactory.h
--- a/Engine/Plugins/Enterprise/LidarPointCloud/Source/LidarPointCloudEditor/Private/LidarPointCloudFactory.h
+++ b/Engine/Plugins/Enterprise/LidarPointCloud/Source/LidarPointCloudEditor/Private/LidarPointCloudFactory.h
@@ -17,7 +17,6 @@ public:
 	virtual FText GetName() const override;
 	virtual FColor GetTypeColor() const override { return FColor(0, 128, 128); }
 	virtual UClass* GetSupportedClass() const override { return ULidarPointCloud::StaticClass(); }
-	virtual bool HasActions(const TArray<UObject*>& InObjects) const override { return true; }
 	virtual void GetActions(const TArray<UObject*>& InObjects, FMenuBuilder& MenuBuilder) override;
 	virtual uint32 GetCategories() override { return EAssetTypeCategories::Misc; }
 	virtual bool IsImportedAsset() const override { return true; }
diff --git a/Engine/Plugins/Experimental/Water/Source/Editor/Private/AssetTypeActions_WaterWaves.h b/Engine/Plugins/Experimental/Water/Source/Editor/Private/AssetTypeActions_WaterWaves.h
--- a/Engine/Plugins/Experimental/Water/Source/Editor/Private/AssetTypeActions_WaterWaves.h
+++ b/Engine/Plugins/Experimental/Water/Source/Editor/Private/AssetTypeActions_WaterWaves.h
@@ -14,8 +14,5 @@ public:
 	virtual UClass* GetSupportedClass() const override;
 	virtual uint32 GetCategories() override;
 	virtual void OpenAssetEditor(const TArray<UObject*>& InObjects, TSharedPtr<class IToolkitHost> EditWithinLevelEditor = TSharedPtr<IToolkitHost>()) override;
-	
-	//virtual bool HasActions(const TArray<UObject*>& InObjects) const override { return true; }
-	//virtual void GetActions(const TArray<UObject*>& InObjects, FMenuBuilder& MenuBuilder) override;
 	// End IAssetTypeActions Interface
 };
\ No newline at end of file
diff --git a/Engine/Plugins/Media/BinkMedia/Source/BinkMediaPlayerEditor/Private/BinkMediaPlayerEditorModule.cpp b/Engine/Plugins/Media/BinkMedia/Source/BinkMediaPlayerEditor/Private/BinkMediaPlayerEditorModule.cpp
--- a/Engine/Plugins/Media/BinkMedia/Source/BinkMediaPlayerEditor/Private/BinkMediaPlayerEditorModule.cpp
+++ b/Engine/Plugins/Media/BinkMedia/Source/BinkMediaPlayerEditor/Private/BinkMediaPlayerEditorModule.cpp
@@ -315,7 +315,6 @@ struct FBinkMediaPlayerActions : FAssetTypeActions_Base
 	virtual FText GetName() const override { return FText::FromString(TEXT("Bink Media Player")); }
 	virtual UClass* GetSupportedClass() const override { return UBinkMediaPlayer::StaticClass(); }
 	virtual FColor GetTypeColor() const override { return FColor::Red; }
-	virtual bool HasActions( const TArray<UObject*>& InObjects ) const override { return  true; }
 
 	TSharedRef<ISlateStyle> Style;
 };
diff --git a/Engine/Plugins/Media/PixelStreaming/Source/PixelStreamingBlueprintEditor/Private/AssetTypeActions_StreamerInput.h b/Engine/Plugins/Media/PixelStreaming/Source/PixelStreamingBlueprintEditor/Private/AssetTypeActions_StreamerInput.h
--- a/Engine/Plugins/Media/PixelStreaming/Source/PixelStreamingBlueprintEditor/Private/AssetTypeActions_StreamerInput.h
+++ b/Engine/Plugins/Media/PixelStreaming/Source/PixelStreamingBlueprintEditor/Private/AssetTypeActions_StreamerInput.h
@@ -14,7 +14,6 @@ public:
 	virtual FText GetName() const override { return NSLOCTEXT("PixelStreaming", "AssetTypeActions_StreamerInput", "Streamer Input Actions"); }
 	virtual FColor GetTypeColor() const override { return FColor(192,64,64); }
 	virtual UClass* GetSupportedClass() const override { return UPixelStreamingStreamerInput::StaticClass(); }
-	virtual bool HasActions ( const TArray<UObject*>& InObjects ) const override { return false; }
 	virtual bool CanFilter() override { return false; }
 	virtual uint32 GetCategories() override { return EAssetTypeCategories::Misc; }
 	virtual bool IsImportedAsset() const override { return false; }
diff --git a/Engine/Plugins/Runtime/ApexDestruction/Source/ApexDestructionEditor/Private/AssetTypeActions_DestructibleMesh.h b/Engine/Plugins/Runtime/ApexDestruction/Source/ApexDestructionEditor/Private/AssetTypeActions_DestructibleMesh.h
--- a/Engine/Plugins/Runtime/ApexDestruction/Source/ApexDestructionEditor/Private/AssetTypeActions_DestructibleMesh.h
+++ b/Engine/Plugins/Runtime/ApexDestruction/Source/ApexDestructionEditor/Private/AssetTypeActions_DestructibleMesh.h
@@ -14,7 +14,6 @@ public:
 	virtual FColor GetTypeColor() const override { return FColor(200,128,128); }
 	virtual UClass* GetSupportedClass() const override;
 	virtual uint32 GetCategories() override { return EAssetTypeCategories::Physics; }
-	virtual bool HasActions( const TArray<UObject*>& InObjects ) const override { return false; }
 	virtual void OpenAssetEditor( const TArray<UObject*>& InObjects, TSharedPtr<class IToolkitHost> EditWithinLevelEditor = TSharedPtr<IToolkitHost>() ) override;
 	virtual bool IsImportedAsset() const override { return false; }
 };
diff --git a/Engine/Plugins/Runtime/HairStrands/Source/HairStrandsEditor/Private/GroomActions.cpp b/Engine/Plugins/Runtime/HairStrands/Source/HairStrandsEditor/Private/GroomActions.cpp
--- a/Engine/Plugins/Runtime/HairStrands/Source/HairStrandsEditor/Private/GroomActions.cpp
+++ b/Engine/Plugins/Runtime/HairStrands/Source/HairStrandsEditor/Private/GroomActions.cpp
@@ -128,11 +128,6 @@ FColor FGroomActions::GetTypeColor() const
 	return FColor::White;
 }
 
-bool FGroomActions::HasActions(const TArray<UObject*>& InObjects) const
-{
-	return true;
-}
-
 void FGroomActions::OpenAssetEditor(const TArray<UObject*>& InObjects, TSharedPtr<IToolkitHost> EditWithinLevelEditor)
 {
 	// #ueent_todo: Will need a custom editor at some point, for now just use the Properties editor
diff --git a/Engine/Plugins/Runtime/HairStrands/Source/HairStrandsEditor/Public/GroomActions.h b/Engine/Plugins/Runtime/HairStrands/Source/HairStrandsEditor/Public/GroomActions.h
--- a/Engine/Plugins/Runtime/HairStrands/Source/HairStrandsEditor/Public/GroomActions.h
+++ b/Engine/Plugins/Runtime/HairStrands/Source/HairStrandsEditor/Public/GroomActions.h
@@ -33,7 +33,6 @@ public:
 	virtual void GetResolvedSourceFilePaths(const TArray<UObject*>& TypeAssets, TArray<FString>& OutSourceFilePaths) const override;
 	virtual UClass* GetSupportedClass() const override;
 	virtual FColor GetTypeColor() const override;
-	virtual bool HasActions(const TArray<UObject*>& InObjects) const override;
 	virtual bool IsImportedAsset() const override { return true; }
 	virtual void OpenAssetEditor(const TArray<UObject*>& InObjects, TSharedPtr<IToolkitHost> EditWithinLevelEditor = TSharedPtr<IToolkitHost>()) override;
 
diff --git a/Engine/Plugins/Runtime/SoundMod/Source/SoundModImporter/Private/AssetTypeActions_SoundMod.h b/Engine/Plugins/Runtime/SoundMod/Source/SoundModImporter/Private/AssetTypeActions_SoundMod.h
--- a/Engine/Plugins/Runtime/SoundMod/Source/SoundModImporter/Private/AssetTypeActions_SoundMod.h
+++ b/Engine/Plugins/Runtime/SoundMod/Source/SoundModImporter/Private/AssetTypeActions_SoundMod.h
@@ -14,7 +14,6 @@ public:
 	virtual FText GetName() const override{ return NSLOCTEXT("AssetTypeActions", "AssetTypeActions_SoundMod", "Sound Mod"); }
 	virtual FColor GetTypeColor() const override{ return FColor(255, 175, 0); }
 	virtual UClass* GetSupportedClass() const override;
-	virtual bool HasActions(const TArray<UObject*>& InObjects) const override{ return true; }
 	virtual void GetActions(const TArray<UObject*>& InObjects, struct FToolMenuSection& Section) override;
 	virtual bool AssetsActivatedOverride(const TArray<UObject*>& InObjects, EAssetTypeActivationMethod::Type ActivationType) override;
 	virtual uint32 GetCategories() override{ return EAssetTypeCategories::Sounds; }
diff --git a/Engine/Plugins/VirtualProduction/CameraCalibration/Source/CameraCalibrationEditor/Private/AssetTypeActions/AssetTypeActions_LensFile.h b/Engine/Plugins/VirtualProduction/CameraCalibration/Source/CameraCalibrationEditor/Private/AssetTypeActions/AssetTypeActions_LensFile.h
--- a/Engine/Plugins/VirtualProduction/CameraCalibration/Source/CameraCalibrationEditor/Private/AssetTypeActions/AssetTypeActions_LensFile.h
+++ b/Engine/Plugins/VirtualProduction/CameraCalibration/Source/CameraCalibrationEditor/Private/AssetTypeActions/AssetTypeActions_LensFile.h
@@ -14,12 +14,7 @@ public:
 	{
 		return FColor(100, 255, 100);
 	}
-
-	virtual bool HasActions(const TArray<UObject*>& InObjects) const override
-	{
-		return false;
-	}
-
+	
 	virtual uint32 GetCategories() override
 	{
 		return EAssetTypeCategories::Misc;
diff --git a/Engine/Plugins/VirtualProduction/LevelSnapshots/Source/LevelSnapshotsEditor/Private/AssetTypeActions/AssetTypeActions_LevelSnapshot.h b/Engine/Plugins/VirtualProduction/LevelSnapshots/Source/LevelSnapshotsEditor/Private/AssetTypeActions/AssetTypeActions_LevelSnapshot.h
--- a/Engine/Plugins/VirtualProduction/LevelSnapshots/Source/LevelSnapshotsEditor/Private/AssetTypeActions/AssetTypeActions_LevelSnapshot.h
+++ b/Engine/Plugins/VirtualProduction/LevelSnapshots/Source/LevelSnapshotsEditor/Private/AssetTypeActions/AssetTypeActions_LevelSnapshot.h
@@ -36,16 +36,6 @@ public:
 	{
 		return FColor(238, 181, 235, 255);
 	}
-	
-	/**
-	 * @brief 
-	 * @param InObjects A list of objects selected in the content browser
-	 * @return Whether or not the asset has context actions (required to return true in order to use registered actions)
-	 */
-	virtual bool HasActions(const TArray<UObject*>& InObjects) const override
-	{
-		return true;
-	}
 
 	// These are defined in source - GetName requires a defined LOCTEXT namespace and GetSupportedClass has an external include
 	
diff --git a/Engine/Plugins/VirtualProduction/RemoteControl/Source/RemoteControlUI/Private/AssetTools/RemoteControlPresetActions.h b/Engine/Plugins/VirtualProduction/RemoteControl/Source/RemoteControlUI/Private/AssetTools/RemoteControlPresetActions.h
--- a/Engine/Plugins/VirtualProduction/RemoteControl/Source/RemoteControlUI/Private/AssetTools/RemoteControlPresetActions.h
+++ b/Engine/Plugins/VirtualProduction/RemoteControl/Source/RemoteControlUI/Private/AssetTools/RemoteControlPresetActions.h
@@ -31,7 +31,6 @@ public:
 	virtual FColor GetTypeColor() const override;
 	virtual void OpenAssetEditor(const TArray<UObject*>& InObjects, TSharedPtr<class IToolkitHost> EditWithinLevelEditor = TSharedPtr<IToolkitHost>()) override;
 	virtual bool CanLocalize() const override { return false; }
-	virtual bool HasActions(const TArray<UObject*>& InObjects) const { return false; }
 	virtual bool ShouldForceWorldCentric() override { return true; }
 	//~ End IAssetTypeActions interface.
 
