Fix potential race condition by moving delegate registration out of async loading thread

diff --git a/Engine/Plugins/Enterprise/DatasmithContent/Source/DatasmithContent/Private/DatasmithSceneActor.cpp b/Engine/Plugins/Enterprise/DatasmithContent/Source/DatasmithContent/Private/DatasmithSceneActor.cpp
--- a/Engine/Plugins/Enterprise/DatasmithContent/Source/DatasmithContent/Private/DatasmithSceneActor.cpp
+++ b/Engine/Plugins/Enterprise/DatasmithContent/Source/DatasmithContent/Private/DatasmithSceneActor.cpp
@@ -13,44 +13,65 @@
 ADatasmithSceneActor::ADatasmithSceneActor()
 {
 #if WITH_EDITOR
-	if (HasAnyFlags(RF_ClassDefaultObject))
+	if (!IsTemplate())
 	{
-		return;
+		if (HasAnyFlags(RF_NeedPostLoad) || GetOuter()->HasAnyFlags(RF_NeedPostLoad))
+		{
+			// Delegate registration is not thread-safe, so we postpone it on PostLoad when coming from loading which could be on another thread
+		}
+		else
+		{
+			RegisterDelegates();
+		}
 	}
+#endif // WITH_EDITOR
+}
+
+#if WITH_EDITOR
+
+void ADatasmithSceneActor::RegisterDelegates()
+{
+	ensureMsgf(IsInGameThread(), TEXT("Potential race condition in ADatasmithSceneActor registering to delegates from non game-thread"));
 
 	FEditorDelegates::MapChange.AddUObject(this, &ADatasmithSceneActor::OnMapChange);
 
-	if ( GEngine )
+	if (GEngine)
 	{
-		OnActorDeletedDelegateHandle = GEngine->OnLevelActorDeleted().AddUObject( this, &ADatasmithSceneActor::OnActorDeleted );
+		OnActorDeletedDelegateHandle = GEngine->OnLevelActorDeleted().AddUObject(this, &ADatasmithSceneActor::OnActorDeleted);
 	}
 
-	if ( GEditor )
+	if (GEditor)
 	{
 		GEditor->GetEditorSubsystem<UImportSubsystem>()->OnAssetPostImport.AddUObject(this, &ADatasmithSceneActor::OnAssetPostImport);
 	}
-#endif // WITH_EDITOR
 }
 
-ADatasmithSceneActor::~ADatasmithSceneActor()
+void ADatasmithSceneActor::UnregisterDelegates()
 {
-#if WITH_EDITOR
-	if (HasAnyFlags(RF_ClassDefaultObject))
-	{
-		return;
-	}
+	ensureMsgf(IsInGameThread(), TEXT("Potential race condition in ADatasmithSceneActor unregistering to delegates from non game-thread"));
 
 	FEditorDelegates::MapChange.RemoveAll(this);
 
-	if ( GEngine )
+	if (GEngine)
 	{
-		GEngine->OnLevelActorDeleted().Remove( OnActorDeletedDelegateHandle );
+		GEngine->OnLevelActorDeleted().Remove(OnActorDeletedDelegateHandle);
 	}
 
-	if ( GEditor )
+	if (GEditor)
 	{
 		GEditor->GetEditorSubsystem<UImportSubsystem>()->OnAssetPostImport.RemoveAll(this);
 	}
+}
+
+#endif
+
+ADatasmithSceneActor::~ADatasmithSceneActor()
+{
+#if WITH_EDITOR
+	if (!IsTemplate())
+	{
+		UnregisterDelegates();
+	}
 #endif // WITH_EDITOR
 }
 
@@ -58,6 +79,11 @@ ADatasmithSceneActor::~ADatasmithSceneActor()
 
 void ADatasmithSceneActor::PostLoad()
 {
+	if (!IsTemplate())
+	{
+		RegisterDelegates();
+	}
+
 	EnsureDatasmithIdsForRelatedActors();
 
 	Super::PostLoad();
diff --git a/Engine/Plugins/Enterprise/DatasmithContent/Source/DatasmithContent/Public/DatasmithSceneActor.h b/Engine/Plugins/Enterprise/DatasmithContent/Source/DatasmithContent/Public/DatasmithSceneActor.h
--- a/Engine/Plugins/Enterprise/DatasmithContent/Source/DatasmithContent/Public/DatasmithSceneActor.h
+++ b/Engine/Plugins/Enterprise/DatasmithContent/Source/DatasmithContent/Public/DatasmithSceneActor.h
@@ -43,6 +43,9 @@ private:
 	// Reattach newly imported actors if they've been moved / cut and pasted
 	void OnAssetPostImport(UFactory* InFactory, UObject* ActorAdded);
 
+	void RegisterDelegates();
+	void UnregisterDelegates();
+
 	FDelegateHandle OnActorDeletedDelegateHandle;
 #endif // WITH_EDITOR
 };
