Fixed frequent crashes reported by CrashReporter
 - Solution: Preventing check on pointer which apparently could be null

diff --git a/Engine/Plugins/Enterprise/DatasmithContent/Source/DatasmithContent/Private/DatasmithSceneActor.cpp b/Engine/Plugins/Enterprise/DatasmithContent/Source/DatasmithContent/Private/DatasmithSceneActor.cpp
--- a/Engine/Plugins/Enterprise/DatasmithContent/Source/DatasmithContent/Private/DatasmithSceneActor.cpp
+++ b/Engine/Plugins/Enterprise/DatasmithContent/Source/DatasmithContent/Private/DatasmithSceneActor.cpp
@@ -57,9 +57,9 @@ void ADatasmithSceneActor::UnregisterDelegates()
 		GEngine->OnLevelActorDeleted().Remove(OnActorDeletedDelegateHandle);
 	}
 
-	if (GEditor)
+	if (UImportSubsystem* ImportSubsystem = (GEditor ? GetEditorSubsystem<UImportSubsystem>() : nullptr))
 	{
-		GEditor->GetEditorSubsystem<UImportSubsystem>()->OnAssetPostImport.RemoveAll(this);
+		ImportSubsystem->OnAssetPostImport.RemoveAll(this);
 	}
 }
 
