Fix resolve issue

diff --git a/Engine/Plugins/Enterprise/DatasmithContent/Source/DatasmithContent/Private/DatasmithSceneActor.cpp b/Engine/Plugins/Enterprise/DatasmithContent/Source/DatasmithContent/Private/DatasmithSceneActor.cpp
--- a/Engine/Plugins/Enterprise/DatasmithContent/Source/DatasmithContent/Private/DatasmithSceneActor.cpp
+++ b/Engine/Plugins/Enterprise/DatasmithContent/Source/DatasmithContent/Private/DatasmithSceneActor.cpp
@@ -57,7 +57,7 @@ void ADatasmithSceneActor::UnregisterDelegates()
 		GEngine->OnLevelActorDeleted().Remove(OnActorDeletedDelegateHandle);
 	}
 
-	if (UImportSubsystem* ImportSubsystem = (GEditor ? GetEditorSubsystem<UImportSubsystem>() : nullptr))
+	if (UImportSubsystem* ImportSubsystem = (GEditor ? GEditor->GetEditorSubsystem<UImportSubsystem>() : nullptr))
 	{
 		ImportSubsystem->OnAssetPostImport.RemoveAll(this);
 	}
