Fixed heap-use-after-free

diff --git a/Engine/Plugins/Enterprise/DatasmithContent/Source/DatasmithContent/Private/DatasmithSceneActor.cpp b/Engine/Plugins/Enterprise/DatasmithContent/Source/DatasmithContent/Private/DatasmithSceneActor.cpp
--- a/Engine/Plugins/Enterprise/DatasmithContent/Source/DatasmithContent/Private/DatasmithSceneActor.cpp
+++ b/Engine/Plugins/Enterprise/DatasmithContent/Source/DatasmithContent/Private/DatasmithSceneActor.cpp
@@ -28,6 +28,18 @@ ADatasmithSceneActor::ADatasmithSceneActor()
 #endif // WITH_EDITOR
 }
 
+void ADatasmithSceneActor::BeginDestroy()
+{
+	Super::BeginDestroy();
+
+#if WITH_EDITOR
+	if (!IsTemplate())
+	{
+		UnregisterDelegates();
+	}
+#endif
+}
+
 #if WITH_EDITOR
 
 void ADatasmithSceneActor::RegisterDelegates()
@@ -66,16 +78,6 @@ void ADatasmithSceneActor::UnregisterDelegates()
 
 #endif
 
-ADatasmithSceneActor::~ADatasmithSceneActor()
-{
-#if WITH_EDITOR
-	if (!IsTemplate())
-	{
-		UnregisterDelegates();
-	}
-#endif // WITH_EDITOR
-}
-
 #if WITH_EDITOR
 
 void ADatasmithSceneActor::PostLoad()
diff --git a/Engine/Plugins/Enterprise/DatasmithContent/Source/DatasmithContent/Public/DatasmithSceneActor.h b/Engine/Plugins/Enterprise/DatasmithContent/Source/DatasmithContent/Public/DatasmithSceneActor.h
--- a/Engine/Plugins/Enterprise/DatasmithContent/Source/DatasmithContent/Public/DatasmithSceneActor.h
+++ b/Engine/Plugins/Enterprise/DatasmithContent/Source/DatasmithContent/Public/DatasmithSceneActor.h
@@ -17,8 +17,6 @@ public:
 
 	ADatasmithSceneActor();
 
-	virtual ~ADatasmithSceneActor();
-
 	UPROPERTY(VisibleAnywhere, Category="Datasmith")
 	TObjectPtr<class UDatasmithScene> Scene;
 
@@ -26,6 +24,9 @@ public:
 	UPROPERTY(VisibleAnywhere, Category="Datasmith", AdvancedDisplay)
 	TMap< FName, TSoftObjectPtr< AActor > > RelatedActors;
 
+	// Clean up delegates before destruction 
+	void BeginDestroy() override;
+
 private:
 #if WITH_EDITOR
 	// For backward compatibility with UE <= 4.21
