[DatasmithImporter] Fixed DirectLink reimport or static meshes when materiasl order changes

- importer wasn't updating slot order

diff --git a/Engine/Plugins/Enterprise/DatasmithContent/Source/DatasmithContent/Private/ObjectTemplates/DatasmithStaticMeshTemplate.cpp b/Engine/Plugins/Enterprise/DatasmithContent/Source/DatasmithContent/Private/ObjectTemplates/DatasmithStaticMeshTemplate.cpp
--- a/Engine/Plugins/Enterprise/DatasmithContent/Source/DatasmithContent/Private/ObjectTemplates/DatasmithStaticMeshTemplate.cpp
+++ b/Engine/Plugins/Enterprise/DatasmithContent/Source/DatasmithContent/Private/ObjectTemplates/DatasmithStaticMeshTemplate.cpp
@@ -255,7 +255,7 @@ UObject* UDatasmithStaticMeshTemplate::UpdateObject( UObject* Destination, bool
 	}
 
 	// Make sure that the StaticMaterials are in the same order than the StaticMeshLODResources.Sections will be after the StaticMesh is built (see BuildVertexBuffer in StaticMeshBuilder.cpp)
-	const int32 NumLODs = StaticMesh->GetNumLODs();
+	const int32 NumLODs = StaticMesh->GetNumSourceModels();
 	for ( int32 LODIndex = 0; LODIndex < NumLODs; ++LODIndex )
 	{
 		const FMeshDescription* MeshDescription = StaticMesh->GetMeshDescription( LODIndex );
