DatasmithContentEditor - Replace usage of IAssetTypeActions_AssetInterface with UAssetDefinition for DatasmithScene

diff --git a/Engine/Plugins/Enterprise/DatasmithContent/Source/DatasmithContentEditor/DatasmithContentEditor.Build.cs b/Engine/Plugins/Enterprise/DatasmithContent/Source/DatasmithContentEditor/DatasmithContentEditor.Build.cs
--- a/Engine/Plugins/Enterprise/DatasmithContent/Source/DatasmithContentEditor/DatasmithContentEditor.Build.cs
+++ b/Engine/Plugins/Enterprise/DatasmithContent/Source/DatasmithContentEditor/DatasmithContentEditor.Build.cs
@@ -10,6 +10,8 @@ namespace UnrealBuildTool.Rules
 			PublicDependencyModuleNames.AddRange(
 				new string[]
 				{
+					"AssetDefinition",
+					"ContentBrowser",
 					"Core",
 					"CoreUObject",
 					"DatasmithContent",
diff --git a/Engine/Plugins/Enterprise/DatasmithContent/Source/DatasmithContentEditor/Private/AssetDefinition_DatasmithScene.cpp b/Engine/Plugins/Enterprise/DatasmithContent/Source/DatasmithContentEditor/Private/AssetDefinition_DatasmithScene.cpp
new file mode 100644
--- /dev/null
+++ b/Engine/Plugins/Enterprise/DatasmithContent/Source/DatasmithContentEditor/Private/AssetDefinition_DatasmithScene.cpp
@@ -0,0 +1,154 @@
+// Copyright Epic Games, Inc. All Rights Reserved.
+
+#include "AssetDefinition_DatasmithScene.h"
+
+#include "ContentBrowserMenuContexts.h"
+#include "DatasmithAssetImportData.h"
+#include "DatasmithContentEditorModule.h"
+#include "DatasmithContentEditorStyle.h"
+#include "DatasmithScene.h"
+
+#include "ToolMenuSection.h"
+
+#define LOCTEXT_NAMESPACE "AssetTypeActions_DatasmithScene"
+
+FText UAssetDefinition_DatasmithScene::GetAssetDisplayName() const
+{
+	return LOCTEXT("AssetTypeActions_DatasmithScene_Name", "Datasmith Scene");
+}
+
+TSoftClassPtr<UObject> UAssetDefinition_DatasmithScene::GetAssetClass() const
+{
+	return UDatasmithScene::StaticClass();
+}
+
+EAssetCommandResult UAssetDefinition_DatasmithScene::OpenAssets(const FAssetOpenArgs& OpenArgs) const
+{
+	FOnCreateDatasmithSceneEditor DatasmithSceneEditorHandler = IDatasmithContentEditorModule::Get().GetDatasmithSceneEditorHandler();
+
+	if (!DatasmithSceneEditorHandler.IsBound() == false)
+	{
+		FSimpleAssetEditor::CreateEditor(EToolkitMode::Standalone, OpenArgs.ToolkitHost, OpenArgs.LoadObjects<UObject>());
+		return EAssetCommandResult::Handled;
+	}
+
+	for (UDatasmithScene* DatasmithScene : OpenArgs.LoadObjects<UDatasmithScene>())
+	{
+		DatasmithSceneEditorHandler.ExecuteIfBound(EToolkitMode::Standalone, OpenArgs.ToolkitHost, DatasmithScene);
+	}
+
+	return EAssetCommandResult::Handled;
+
+}
+
+TConstArrayView<FAssetCategoryPath> UAssetDefinition_DatasmithScene::GetAssetCategories() const
+{
+	static const auto Categories = { FAssetCategoryPath(NSLOCTEXT("DatasmithContentEditorModule", "DatasmithContentAssetCategory", "Datasmith")) };
+	return Categories;
+}
+
+namespace MenuExtension_DatasmithScene
+{
+
+	void OpenAssetEditor(const TArray<UObject*>& InObjects, TSharedPtr<class IToolkitHost> EditWithinLevelEditor)
+	{
+		if (InObjects.Num() == 0)
+		{
+			return;
+		}
+
+		FOnCreateDatasmithSceneEditor DatasmithSceneEditorHandler = IDatasmithContentEditorModule::Get().GetDatasmithSceneEditorHandler();
+
+		if (DatasmithSceneEditorHandler.IsBound() == false)
+		{
+			FSimpleAssetEditor::CreateEditor(EToolkitMode::Standalone, EditWithinLevelEditor, InObjects);
+			return;
+		}
+
+		for (UObject* Object : InObjects)
+		{
+			UDatasmithScene* DatasmithScene = Cast<UDatasmithScene>(Object);
+			if (DatasmithScene != nullptr)
+			{
+				DatasmithSceneEditorHandler.ExecuteIfBound(EToolkitMode::Standalone, EditWithinLevelEditor, DatasmithScene);
+			}
+		}
+	}
+
+	void ExecuteToggleDirectLinkAutoReimport(const TArray<TWeakObjectPtr<UDatasmithScene>>& Scenes, bool bEnabled)
+	{
+		const IDatasmithContentEditorModule& DatasmithContentEditorModule = IDatasmithContentEditorModule::Get();
+
+		for (const TWeakObjectPtr<class UDatasmithScene>& Scene : Scenes)
+		{
+			if (Scene.IsValid())
+			{
+				DatasmithContentEditorModule.SetAssetAutoReimport(Scene.Get(), bEnabled);
+			}
+		}
+	}
+
+	void FilterByDirectLinkAutoReimportSupport(TArray<UDatasmithScene*>& Scenes)
+	{
+		const IDatasmithContentEditorModule& DatasmithContentEditorModule = IDatasmithContentEditorModule::Get();
+
+		for (int32 SceneIndex = Scenes.Num() - 1; SceneIndex >= 0; --SceneIndex)
+		{
+			const bool bIsReimportAvailable = DatasmithContentEditorModule.IsAssetAutoReimportAvailable(Scenes[SceneIndex]).Get(false);
+			if (bIsReimportAvailable)
+			{
+				continue;
+			}
+			Scenes.RemoveAtSwap(SceneIndex);
+		}
+	}
+
+	static FDelayedAutoRegisterHelper DelayedAutoRegister(EDelayedRegisterRunPhase::EndOfEngineInit, []{ 
+		UToolMenus::RegisterStartupCallback(FSimpleMulticastDelegate::FDelegate::CreateLambda([]()
+		{
+			FToolMenuOwnerScoped OwnerScoped(UE_MODULE_NAME);
+			UToolMenu* Menu = UE::ContentBrowser::ExtendToolMenu_AssetContextMenu(UDatasmithScene::StaticClass());
+	        
+			FToolMenuSection& Section = Menu->FindOrAddSection("GetAssetActions");
+			Section.AddDynamicEntry(NAME_None, FNewToolMenuSectionDelegate::CreateLambda([](FToolMenuSection& InSection)
+			{
+				if (const UContentBrowserAssetContextMenuContext* CBContext = UContentBrowserAssetContextMenuContext::FindContextWithAssets(InSection))
+				{
+					// todo: investigate if it's possible to check AutoReimport availability/state without loading the DatasmithScene asset
+					// so that loading would be done only when action(ExecuteToggleDirectLinkAutoReimport) is executed
+					TArray<UDatasmithScene*> Scenes = CBContext->LoadSelectedObjects<UDatasmithScene>();
+
+					FilterByDirectLinkAutoReimportSupport(Scenes);
+
+					const bool bHasCanExecuteAutoReimport = Scenes.Num() > 0;
+					
+					bool bHasAutoReimportEnabled = false;
+					const IDatasmithContentEditorModule& DatasmithContentEditorModule = IDatasmithContentEditorModule::Get();
+					for (UDatasmithScene* Scene : Scenes)
+					{
+						bHasAutoReimportEnabled = DatasmithContentEditorModule.IsAssetAutoReimportEnabled(Scene).Get(false);
+						if (bHasAutoReimportEnabled)
+						{
+							break;
+						}
+					}
+
+					const TAttribute<FText> Label = LOCTEXT("DatasmithScene_ToggleDirectLinkAutoReimport", "Direct Link Auto-Reimport");
+					const TAttribute<FText> ToolTip = LOCTEXT("DatasmithScene_ToggleDirectLinkAutoReimportTooltip", "Toggle Direct Link Auto-Reimport for all selected Datasmith Scenes.");
+					const FSlateIcon Icon = FSlateIcon(FDatasmithContentEditorStyle::GetStyleSetName(), "DatasmithContent.AutoReimportGrayscale");
+
+					FToolUIAction UIAction;
+					UIAction.ExecuteAction = FToolMenuExecuteAction::CreateLambda([Scenes = TArray<TWeakObjectPtr<UDatasmithScene>>(Scenes), bHasAutoReimportEnabled](const FToolMenuContext& Context) { ExecuteToggleDirectLinkAutoReimport(Scenes, !bHasAutoReimportEnabled); });
+					UIAction.CanExecuteAction = FToolMenuCanExecuteAction::CreateLambda([bHasCanExecuteAutoReimport](const FToolMenuContext& Context) { return bHasCanExecuteAutoReimport; });
+					UIAction.GetActionCheckState = FToolMenuGetActionCheckState::CreateLambda([bHasAutoReimportEnabled](const FToolMenuContext& Context) { return bHasAutoReimportEnabled ? ECheckBoxState::Checked : ECheckBoxState::Unchecked; });
+					
+					InSection.AddMenuEntry("DatasmithScene_ToggleDirectLinkAutoReimport", Label, ToolTip, Icon, UIAction, EUserInterfaceActionType::Check);
+
+				}
+			}));
+		}));
+	});
+}
+
+
+#undef LOCTEXT_NAMESPACE
diff --git a/Engine/Plugins/Enterprise/DatasmithContent/Source/DatasmithContentEditor/Private/AssetDefinition_DatasmithScene.h b/Engine/Plugins/Enterprise/DatasmithContent/Source/DatasmithContentEditor/Private/AssetDefinition_DatasmithScene.h
new file mode 100644
--- /dev/null
+++ b/Engine/Plugins/Enterprise/DatasmithContent/Source/DatasmithContentEditor/Private/AssetDefinition_DatasmithScene.h
@@ -0,0 +1,25 @@
+// Copyright Epic Games, Inc. All Rights Reserved.
+
+#pragma once
+
+#include "AssetDefinitionDefault.h"
+#include "AssetTypeActions_Base.h"
+
+#include "AssetDefinition_DatasmithScene.generated.h"
+
+
+UCLASS()
+class UAssetDefinition_DatasmithScene : public UAssetDefinitionDefault
+{
+	GENERATED_BODY()
+
+public:
+	// UAssetDefinition Begin
+	virtual FText GetAssetDisplayName() const override;
+	virtual FLinearColor GetAssetColor() const override { return FColor(0, 0, 0); }
+	virtual TSoftClassPtr<UObject> GetAssetClass() const override;
+	virtual EAssetCommandResult OpenAssets(const FAssetOpenArgs& OpenArgs) const override;
+	virtual TConstArrayView<FAssetCategoryPath> GetAssetCategories() const override;
+	virtual bool CanImport() const override { return true; }
+	// UAssetDefinition End
+};
diff --git a/Engine/Plugins/Enterprise/DatasmithContent/Source/DatasmithContentEditor/Private/AssetTypeActions_DatasmithScene.cpp b/Engine/Plugins/Enterprise/DatasmithContent/Source/DatasmithContentEditor/Private/AssetTypeActions_DatasmithScene.cpp
deleted file mode 100644
diff --git a/Engine/Plugins/Enterprise/DatasmithContent/Source/DatasmithContentEditor/Private/AssetTypeActions_DatasmithScene.h b/Engine/Plugins/Enterprise/DatasmithContent/Source/DatasmithContentEditor/Private/AssetTypeActions_DatasmithScene.h
deleted file mode 100644
diff --git a/Engine/Plugins/Enterprise/DatasmithContent/Source/DatasmithContentEditor/Private/DatasmithContentEditorModule.cpp b/Engine/Plugins/Enterprise/DatasmithContent/Source/DatasmithContentEditor/Private/DatasmithContentEditorModule.cpp
--- a/Engine/Plugins/Enterprise/DatasmithContent/Source/DatasmithContentEditor/Private/DatasmithContentEditorModule.cpp
+++ b/Engine/Plugins/Enterprise/DatasmithContent/Source/DatasmithContentEditor/Private/DatasmithContentEditorModule.cpp
@@ -2,7 +2,6 @@
 
 #include "DatasmithContentEditorModule.h"
 
-#include "AssetTypeActions_DatasmithScene.h"
 #include "DatasmithAreaLightActorDetailsPanel.h"
 #include "DatasmithContentEditorStyle.h"
 #include "DatasmithImportInfoCustomization.h"
@@ -19,8 +18,6 @@
 
 #define LOCTEXT_NAMESPACE "DatasmithContentEditorModule"
 
-EAssetTypeCategories::Type IDatasmithContentEditorModule::DatasmithAssetCategoryBit;
-
 /**
  * DatasmithContent module implementation (private)
  */
@@ -34,15 +31,6 @@ public:
 		PropertyModule.RegisterCustomClassLayout( TEXT("DatasmithSceneActor"), FOnGetDetailCustomizationInstance::CreateStatic( &FDatasmithSceneActorDetailsPanel::MakeInstance ) );
 		PropertyModule.RegisterCustomClassLayout( TEXT("DatasmithAreaLightActor"), FOnGetDetailCustomizationInstance::CreateStatic(&FDatasmithAreaLightActorDetailsPanel::MakeInstance));
 
-		// Register Datasmith asset category to group asset type actions related to Datasmith together
-		IAssetTools& AssetTools = FModuleManager::LoadModuleChecked<FAssetToolsModule>("AssetTools").Get();
-		DatasmithAssetCategoryBit = AssetTools.RegisterAdvancedAssetCategory(FName(TEXT("Datasmith")), LOCTEXT("DatasmithContentAssetCategory", "Datasmith"));
-
-		// Register asset type actions for DatasmithScene class
-		TSharedPtr<FAssetTypeActions_DatasmithScene> DatasmithSceneAssetTypeAction = MakeShareable(new FAssetTypeActions_DatasmithScene);
-		AssetTools.RegisterAssetTypeActions(DatasmithSceneAssetTypeAction.ToSharedRef());
-		AssetTypeActionsArray.Add(DatasmithSceneAssetTypeAction);
-
 		FDatasmithContentEditorStyle::Initialize();
 
 		RegisterDetailCustomization();
@@ -58,17 +46,6 @@ public:
 			PropertyModule.NotifyCustomizationModuleChanged();
 		}
 
-		// Unregister asset type actions
-		if (FModuleManager::Get().IsModuleLoaded(TEXT("AssetTools")))
-		{
-			IAssetTools& AssetTools = FModuleManager::LoadModuleChecked<FAssetToolsModule>("AssetTools").Get();
-			for (TSharedPtr<FAssetTypeActions_Base>& AssetTypeActions : AssetTypeActionsArray)
-			{
-				AssetTools.UnregisterAssetTypeActions(AssetTypeActions.ToSharedRef());
-			}
-		}
-		AssetTypeActionsArray.Empty();
-
 		// Shutdown style set associated with datasmith content
 		FDatasmithContentEditorStyle::Shutdown();
 
@@ -275,7 +252,6 @@ private:
 private:
 	FOnSpawnDatasmithSceneActors SpawnActorsDelegate;
 	FOnCreateDatasmithSceneEditor CreateDatasmithSceneEditorDelegate;
-	TArray<TSharedPtr<FAssetTypeActions_Base>> AssetTypeActionsArray;
 	TMap<const void*, FImporterDescription> DatasmithImporterMap;
 
 	FOnSetAssetAutoReimport SetAssetAutoReimportHandler;
diff --git a/Engine/Plugins/Enterprise/DatasmithContent/Source/DatasmithContentEditor/Public/DatasmithContentEditorModule.h b/Engine/Plugins/Enterprise/DatasmithContent/Source/DatasmithContentEditor/Public/DatasmithContentEditorModule.h
--- a/Engine/Plugins/Enterprise/DatasmithContent/Source/DatasmithContentEditor/Public/DatasmithContentEditorModule.h
+++ b/Engine/Plugins/Enterprise/DatasmithContent/Source/DatasmithContentEditor/Public/DatasmithContentEditorModule.h
@@ -105,9 +105,6 @@ public:
 	virtual void UnregisterDatasmithImporter(const void* Registrar) = 0;
 	virtual TArray<FImporterDescription> GetDatasmithImporters() = 0;
 
-	/** Category bit associated with Datasmith related content */
-	static DATASMITHCONTENTEDITOR_API EAssetTypeCategories::Type DatasmithAssetCategoryBit;
-
 	/**
 	 * Delegate to enable or disable AutoReimport of an asset.
 	 */
