Fixed DatasmithScene not opening its asset editor because of a typo when looking if its delegate is bound or not.

diff --git a/Engine/Plugins/Enterprise/DatasmithContent/Source/DatasmithContentEditor/Private/AssetDefinition_DatasmithScene.cpp b/Engine/Plugins/Enterprise/DatasmithContent/Source/DatasmithContentEditor/Private/AssetDefinition_DatasmithScene.cpp
--- a/Engine/Plugins/Enterprise/DatasmithContent/Source/DatasmithContentEditor/Private/AssetDefinition_DatasmithScene.cpp
+++ b/Engine/Plugins/Enterprise/DatasmithContent/Source/DatasmithContentEditor/Private/AssetDefinition_DatasmithScene.cpp
@@ -26,7 +26,7 @@ EAssetCommandResult UAssetDefinition_DatasmithScene::OpenAssets(const FAssetOpen
 {
 	FOnCreateDatasmithSceneEditor DatasmithSceneEditorHandler = IDatasmithContentEditorModule::Get().GetDatasmithSceneEditorHandler();
 
-	if (!DatasmithSceneEditorHandler.IsBound() == false)
+	if (DatasmithSceneEditorHandler.IsBound() == false)
 	{
 		FSimpleAssetEditor::CreateEditor(EToolkitMode::Standalone, OpenArgs.ToolkitHost, OpenArgs.LoadObjects<UObject>());
 		return EAssetCommandResult::Handled;
