Finalized addition of geometrical and stitching tolerances to tessellation options

#rnx

diff --git a/Engine/Plugins/Enterprise/DatasmithCADImporter/Source/CADInterfaces/Private/TechSoftFileParserCADKernelTessellator.cpp b/Engine/Plugins/Enterprise/DatasmithCADImporter/Source/CADInterfaces/Private/TechSoftFileParserCADKernelTessellator.cpp
--- a/Engine/Plugins/Enterprise/DatasmithCADImporter/Source/CADInterfaces/Private/TechSoftFileParserCADKernelTessellator.cpp
+++ b/Engine/Plugins/Enterprise/DatasmithCADImporter/Source/CADInterfaces/Private/TechSoftFileParserCADKernelTessellator.cpp
@@ -300,9 +300,10 @@ void FTechSoftFileParserCADKernelTessellator::MeshAndGetTessellation(UE::CADKern
 	// Tessellate the body
 	TSharedRef<FModelMesh> CADKernelModelMesh = FEntity::MakeShared<FModelMesh>();
 
-	FCADKernelTools::DefineMeshCriteria(CADKernelModelMesh.Get(), CADFileData.GetImportParameters(), CADKernelSession.GetGeometricTolerance());
+	const double GeometricTolerance = CADKernelSession.GetGeometricTolerance();
+
+	FCADKernelTools::DefineMeshCriteria(CADKernelModelMesh.Get(), CADFileData.GetImportParameters(), GeometricTolerance);
 
-	const double GeometricTolerance = FImportParameters::GStitchingTolerance * 10; // cm to mm
 	const bool bActivateThinZoneMeshing = FImportParameters::bGActivateThinZoneMeshing;
 	FMesher Mesher(*CADKernelModelMesh, GeometricTolerance, bActivateThinZoneMeshing);
 	Mesher.MeshEntity(CADKernelBody);
diff --git a/Engine/Plugins/Enterprise/DatasmithCADImporter/Source/CADKernelSurface/CADKernelSurface.Build.cs b/Engine/Plugins/Enterprise/DatasmithCADImporter/Source/CADKernelSurface/CADKernelSurface.Build.cs
--- a/Engine/Plugins/Enterprise/DatasmithCADImporter/Source/CADKernelSurface/CADKernelSurface.Build.cs
+++ b/Engine/Plugins/Enterprise/DatasmithCADImporter/Source/CADKernelSurface/CADKernelSurface.Build.cs
@@ -8,6 +8,8 @@ namespace UnrealBuildTool.Rules
 	{
 		public CADKernelSurface(ReadOnlyTargetRules Target) : base(Target)
 		{
+			OptimizeCode = CodeOptimization.Never;
+
 			PublicDependencyModuleNames.AddRange(
 				new string[]
 				{
diff --git a/Engine/Plugins/Enterprise/DatasmithCADImporter/Source/CADKernelSurface/Private/CADKernelSurfaceExtension.cpp b/Engine/Plugins/Enterprise/DatasmithCADImporter/Source/CADKernelSurface/Private/CADKernelSurfaceExtension.cpp
--- a/Engine/Plugins/Enterprise/DatasmithCADImporter/Source/CADKernelSurface/Private/CADKernelSurfaceExtension.cpp
+++ b/Engine/Plugins/Enterprise/DatasmithCADImporter/Source/CADKernelSurface/Private/CADKernelSurfaceExtension.cpp
@@ -45,7 +45,6 @@ bool UCADKernelParametricSurfaceData::Tessellate(UStaticMesh& StaticMesh, const
 	CadMeshParameters.SymmetricNormal = (FVector3f) MeshParameters.SymmetricNormal;
 	CadMeshParameters.SymmetricOrigin = (FVector3f) MeshParameters.SymmetricOrigin;
 
-	CADLibrary::FMeshConversionContext MeshConversionContext(ImportParameters, CadMeshParameters);
 
 	// Previous MeshDescription is get to be able to create a new one with the same order of PolygonGroup (the matching of color and partition is currently based on their order)
 	if (FMeshDescription* DestinationMeshDescription = StaticMesh.GetMeshDescription(0))
@@ -54,13 +53,7 @@ bool UCADKernelParametricSurfaceData::Tessellate(UStaticMesh& StaticMesh, const
 		FStaticMeshAttributes MeshDescriptionAttributes(MeshDescription);
 		MeshDescriptionAttributes.Register();
 
-		if (RetessellateOptions.RetessellationRule == EDatasmithCADRetessellationRule::SkipDeletedSurfaces)
-		{
-			CADLibrary::GetExistingPatches(*DestinationMeshDescription, MeshConversionContext.PatchesToMesh);
-		}
-
-		const double GeometricTolerance = 0.01; // mm
-		TSharedRef<FSession> CADKernelSession = MakeShared<FSession>(GeometricTolerance);
+		TSharedRef<FSession> CADKernelSession = MakeShared<FSession>(RetessellateOptions.GetGeometricTolerance(true));
 		CADKernelSession->AddDatabase(RawData);
 
 		FModel& CADKernelModel = CADKernelSession->GetModel();
@@ -69,6 +62,13 @@ bool UCADKernelParametricSurfaceData::Tessellate(UStaticMesh& StaticMesh, const
 			return bSuccessfulTessellation;
 		}
 
+		CADLibrary::FMeshConversionContext MeshConversionContext(ImportParameters, CadMeshParameters, CADKernelSession->GetGeometricTolerance());
+
+		if (RetessellateOptions.RetessellationRule == EDatasmithCADRetessellationRule::SkipDeletedSurfaces)
+		{
+			CADLibrary::GetExistingPatches(*DestinationMeshDescription, MeshConversionContext.PatchesToMesh);
+		}
+
 		if(CADLibrary::FCADKernelTools::Tessellate(CADKernelModel, MeshConversionContext, MeshDescription))
 		{
 			// To update the SectionInfoMap 
diff --git a/Engine/Plugins/Enterprise/DatasmithCADImporter/Source/CADKernelSurface/Public/CADModelToCADKernelConverterBase.h b/Engine/Plugins/Enterprise/DatasmithCADImporter/Source/CADKernelSurface/Public/CADModelToCADKernelConverterBase.h
--- a/Engine/Plugins/Enterprise/DatasmithCADImporter/Source/CADKernelSurface/Public/CADModelToCADKernelConverterBase.h
+++ b/Engine/Plugins/Enterprise/DatasmithCADImporter/Source/CADKernelSurface/Public/CADModelToCADKernelConverterBase.h
@@ -4,6 +4,7 @@
 
 #include "CoreMinimal.h"
 
+#include "DatasmithImportOptions.h"
 #include "CADKernelSurfaceExtension.h"
 #include "CADKernelTools.h"
 #include "CADModelConverter.h"
@@ -22,12 +23,9 @@ class FCADModelToCADKernelConverterBase : public CADLibrary::ICADModelConverter
 {
 public:
 
-	FCADModelToCADKernelConverterBase(CADLibrary::FImportParameters InImportParameters)
-		: CADKernelSession(0.01)
+	FCADModelToCADKernelConverterBase(const CADLibrary::FImportParameters& InImportParameters)
+		: CADKernelSession(GeometricTolerance)
 		, ImportParameters(InImportParameters)
-		, GeometricTolerance(0.01)
-		, SquareTolerance(GeometricTolerance* GeometricTolerance)
-		, EdgeLengthTolerance(2 * GeometricTolerance)
 	{
 	}
 
@@ -42,7 +40,6 @@ public:
 		// Apply stitching if applicable
 		if(ImportParameters.GetStitchingTechnique() != StitchingNone)
 		{
-			const double StitchingTolerance = FImportParameters::GStitchingTolerance * 10.; //CM to MM
 			UE::CADKernel::FTopomakerOptions TopomakerOptions((UE::CADKernel::ESewOption)SewOption::GetFromImportParameters(), StitchingTolerance, FImportParameters::GStitchingForceFactor);
 
 			UE::CADKernel::FTopomaker Topomaker(CADKernelSession, TopomakerOptions);
@@ -71,7 +68,7 @@ public:
 	{
 		UE::CADKernel::FModel& Model = CADKernelSession.GetModel();
 		
-		CADLibrary::FMeshConversionContext Context(ImportParameters, InMeshParameters);
+		CADLibrary::FMeshConversionContext Context(ImportParameters, InMeshParameters, CADKernelSession.GetGeometricTolerance());
 
 		return CADLibrary::FCADKernelTools::Tessellate(Model, Context, OutMeshDescription);
 	}
@@ -93,13 +90,24 @@ public:
 		CADKernelSurface::AddSurfaceDataForMesh(InFilePath, ImportParameters, InMeshParameters, InTessellationOptions, OutMeshPayload);
 	}
 
+protected:
+	void SetTolerances(double InGeometryTolerance = 0.01, double InStitchingTolerance = 0.01)
+	{
+		GeometricTolerance = InGeometryTolerance;
+		SquareTolerance = GeometricTolerance * GeometricTolerance;
+		StitchingTolerance = InStitchingTolerance;
+		EdgeLengthTolerance = 2. * GeometricTolerance;
+		CADKernelSession.SetGeometricTolerance(InGeometryTolerance);
+	}
+
 protected:
 
 	UE::CADKernel::FSession CADKernelSession;
 
 	CADLibrary::FImportParameters ImportParameters;
-	double GeometricTolerance;
-	double SquareTolerance;
-	double EdgeLengthTolerance;
+	double GeometricTolerance = 0.01;
+	double SquareTolerance = 0.0001;
+	double EdgeLengthTolerance = 0.02;
+	double StitchingTolerance = 0.01;
 };
 
diff --git a/Engine/Plugins/Enterprise/DatasmithCADImporter/Source/DatasmithOpenNurbsTranslator/Private/DatasmithOpenNurbsTranslator.cpp b/Engine/Plugins/Enterprise/DatasmithCADImporter/Source/DatasmithOpenNurbsTranslator/Private/DatasmithOpenNurbsTranslator.cpp
--- a/Engine/Plugins/Enterprise/DatasmithCADImporter/Source/DatasmithOpenNurbsTranslator/Private/DatasmithOpenNurbsTranslator.cpp
+++ b/Engine/Plugins/Enterprise/DatasmithCADImporter/Source/DatasmithOpenNurbsTranslator/Private/DatasmithOpenNurbsTranslator.cpp
@@ -643,21 +643,6 @@ public:
 		{
 			TranslationCache = MakeShared<FTranslationCache>();
 		}
-
-		CADLibrary::FImportParameters ImportParameters(FDatasmithUtils::EModelCoordSystem::ZUp_RightHanded_FBXLegacy);
-
-		if (CADLibrary::FImportParameters::bGDisableCADKernelTessellation)
-		{
-			TSharedRef<FOpenNurbsBRepToTechSoftConverter> OpenNurbsBRepToTechSoftConverter = MakeShared<FOpenNurbsBRepToTechSoftConverter>(ImportParameters);
-			CADModelConverter = OpenNurbsBRepToTechSoftConverter;
-			OpenNurbsBRepConverter = OpenNurbsBRepToTechSoftConverter;
-		}
-		else
-		{
-			TSharedRef<FOpenNurbsBRepToCADKernelConverter> OpenNurbsBRepToCADKernelConverter = MakeShared<FOpenNurbsBRepToCADKernelConverter>(ImportParameters);
-			CADModelConverter = OpenNurbsBRepToCADKernelConverter;
-			OpenNurbsBRepConverter = OpenNurbsBRepToCADKernelConverter;
-		}
 	}
 
 	~FOpenNurbsTranslatorImpl()
@@ -3228,6 +3213,21 @@ void FOpenNurbsTranslatorImpl::SetOpenNurbsOptions(const FDatasmithOpenNurbsOpti
 	OpenNurbsOptions = Options;
 	OpenNurbsOptionsHash = OpenNurbsOptions.GetHash();
 
+	const CADLibrary::FImportParameters ImportParameters(FDatasmithUtils::EModelCoordSystem::ZUp_RightHanded_FBXLegacy);
+
+	if (CADLibrary::FImportParameters::bGDisableCADKernelTessellation)
+	{
+		TSharedRef<FOpenNurbsBRepToTechSoftConverter> OpenNurbsBRepToTechSoftConverter = MakeShared<FOpenNurbsBRepToTechSoftConverter>(ImportParameters);
+		CADModelConverter = OpenNurbsBRepToTechSoftConverter;
+		OpenNurbsBRepConverter = OpenNurbsBRepToTechSoftConverter;
+	}
+	else
+	{
+		TSharedRef<FOpenNurbsBRepToCADKernelConverter> OpenNurbsBRepToCADKernelConverter = MakeShared<FOpenNurbsBRepToCADKernelConverter>(ImportParameters, Options);
+		CADModelConverter = OpenNurbsBRepToCADKernelConverter;
+		OpenNurbsBRepConverter = OpenNurbsBRepToCADKernelConverter;
+	}
+
 	for (FOpenNurbsTranslatorImpl* ChildTranslator : ChildTranslators)
 	{
 		ChildTranslator->SetOpenNurbsOptions(OpenNurbsOptions);
diff --git a/Engine/Plugins/Enterprise/DatasmithCADImporter/Source/DatasmithOpenNurbsTranslator/Private/OpenNurbsBRepToCADKernelConverter.h b/Engine/Plugins/Enterprise/DatasmithCADImporter/Source/DatasmithOpenNurbsTranslator/Private/OpenNurbsBRepToCADKernelConverter.h
--- a/Engine/Plugins/Enterprise/DatasmithCADImporter/Source/DatasmithOpenNurbsTranslator/Private/OpenNurbsBRepToCADKernelConverter.h
+++ b/Engine/Plugins/Enterprise/DatasmithCADImporter/Source/DatasmithOpenNurbsTranslator/Private/OpenNurbsBRepToCADKernelConverter.h
@@ -32,9 +32,10 @@ class FOpenNurbsBRepToCADKernelConverter : public FCADModelToCADKernelConverterB
 {
 public:
 
-	FOpenNurbsBRepToCADKernelConverter(const CADLibrary::FImportParameters& InImportParameters)
+	FOpenNurbsBRepToCADKernelConverter(const CADLibrary::FImportParameters& InImportParameters, const FDatasmithTessellationOptions& Options)
 		: FCADModelToCADKernelConverterBase(InImportParameters)
 	{
+		SetTolerances(Options.GetGeometricTolerance(true), Options.GetStitchingTolerance(true));
 	}
 
 	virtual ~FOpenNurbsBRepToCADKernelConverter() = default;
@@ -47,8 +48,6 @@ public:
 	 */
 	bool AddBRep(ON_Brep& Brep, const ON_3dVector& Offset);
 
-	static TSharedPtr<FOpenNurbsBRepToCADKernelConverter> GetSharedSession();
-
 private:
 	TSharedPtr<UE::CADKernel::FTopologicalFace> AddFace(const ON_BrepFace& OpenNurbsFace);
 	TSharedPtr<UE::CADKernel::FSurface> AddSurface(ON_NurbsSurface& Surface);
diff --git a/Engine/Plugins/Enterprise/DatasmithCADImporter/Source/ParametricSurface/Private/TechSoft/TechSoftParametricSurface.cpp b/Engine/Plugins/Enterprise/DatasmithCADImporter/Source/ParametricSurface/Private/TechSoft/TechSoftParametricSurface.cpp
--- a/Engine/Plugins/Enterprise/DatasmithCADImporter/Source/ParametricSurface/Private/TechSoft/TechSoftParametricSurface.cpp
+++ b/Engine/Plugins/Enterprise/DatasmithCADImporter/Source/ParametricSurface/Private/TechSoft/TechSoftParametricSurface.cpp
@@ -38,7 +38,7 @@ bool UTechSoftParametricSurfaceData::Tessellate(UStaticMesh& StaticMesh, const F
 		CADLibrary::FImportParameters ImportParameters((FDatasmithUtils::EModelCoordSystem)SceneParameters.ModelCoordSys);
 		ImportParameters.SetTesselationParameters(RetessellateOptions.ChordTolerance, RetessellateOptions.MaxEdgeLength, RetessellateOptions.NormalTolerance, (CADLibrary::EStitchingTechnique)RetessellateOptions.StitchingTechnique);
 
-		FMeshConversionContext Context(ImportParameters, MeshParameters);
+		FMeshConversionContext Context(ImportParameters, MeshParameters, RetessellateOptions.GetGeometricTolerance(true));
 
 		CADLibrary::FTechSoftInterface& TechSoftInterface = CADLibrary::FTechSoftInterface::Get();
 		bSuccessfulTessellation = TechSoftInterface.InitializeKernel(*FPaths::EnginePluginsDir());
diff --git a/Engine/Plugins/Enterprise/DatasmithCADImporter/Source/WireInterface/Private/AliasModelToCADKernelConverter.cpp b/Engine/Plugins/Enterprise/DatasmithCADImporter/Source/WireInterface/Private/AliasModelToCADKernelConverter.cpp
--- a/Engine/Plugins/Enterprise/DatasmithCADImporter/Source/WireInterface/Private/AliasModelToCADKernelConverter.cpp
+++ b/Engine/Plugins/Enterprise/DatasmithCADImporter/Source/WireInterface/Private/AliasModelToCADKernelConverter.cpp
@@ -116,11 +116,7 @@ namespace AliasToCADKernelUtils
 FAliasModelToCADKernelConverter::FAliasModelToCADKernelConverter(const FDatasmithTessellationOptions& Options, CADLibrary::FImportParameters InImportParameters)
 	: FCADModelToCADKernelConverterBase(InImportParameters)
 {
-	GeometricTolerance = UE_TO_CADKERNEL(Options.GeometricTolerance);
-	SquareTolerance = GeometricTolerance * GeometricTolerance;
-	CADKernelSession.SetGeometricTolerance(GeometricTolerance);
-	StitchingTolerance = UE_TO_CADKERNEL(Options.StitchingTolerance);
-	EdgeLengthTolerance = 2. * GeometricTolerance;
+	SetTolerances(Options.GetGeometricTolerance(true), Options.GetStitchingTolerance(true));
 }
 
 TSharedPtr<UE::CADKernel::FTopologicalEdge> FAliasModelToCADKernelConverter::AddEdge(const AlTrimCurve& AliasTrimCurve, TSharedPtr<UE::CADKernel::FSurface>& CarrierSurface)
diff --git a/Engine/Plugins/Enterprise/DatasmithCADImporter/Source/WireInterface/Private/AliasModelToCADKernelConverter.h b/Engine/Plugins/Enterprise/DatasmithCADImporter/Source/WireInterface/Private/AliasModelToCADKernelConverter.h
--- a/Engine/Plugins/Enterprise/DatasmithCADImporter/Source/WireInterface/Private/AliasModelToCADKernelConverter.h
+++ b/Engine/Plugins/Enterprise/DatasmithCADImporter/Source/WireInterface/Private/AliasModelToCADKernelConverter.h
@@ -65,7 +65,6 @@ protected:
 protected:
 	int32 LastFaceId = 1;
 	TMap<void*, TSharedPtr<UE::CADKernel::FTopologicalEdge>>  AlEdge2CADKernelEdge;
-	double StitchingTolerance;
 };
 
 }
diff --git a/Engine/Plugins/Enterprise/DatasmithContent/Source/DatasmithContent/Public/DatasmithImportOptions.h b/Engine/Plugins/Enterprise/DatasmithContent/Source/DatasmithContent/Public/DatasmithImportOptions.h
--- a/Engine/Plugins/Enterprise/DatasmithContent/Source/DatasmithContent/Public/DatasmithImportOptions.h
+++ b/Engine/Plugins/Enterprise/DatasmithContent/Source/DatasmithContent/Public/DatasmithImportOptions.h
@@ -267,10 +267,14 @@ struct DATASMITHCONTENT_API FDatasmithTessellationOptions
 	UPROPERTY(config, EditAnywhere, BlueprintReadWrite, AdvancedDisplay, Category = "Geometry & Tessellation Options", meta = (ToolTip = "Stitching technique applied on model before tessellation. Sewing could impact number of objects."))
 	EDatasmithCADStitchingTechnique StitchingTechnique;
 
+	bool bUseCADKernel = false;
+
+protected:
 	/**
 	 * Tolerance used to determine if a surface should be tessellate or not.
 	 * Any surface which is narrower than the geometric tolerance
 	 * in one of the iso direction will not be tessellated
+	 * Value is in centimeter
 	 */
 	UPROPERTY(config, EditAnywhere, BlueprintReadWrite, AdvancedDisplay, Category = "Geometry & Tessellation Options", meta = (Units = cm, ToolTip = "Tolerance used to determine if a surface should be tessellated or not."))
 	double GeometricTolerance = 0.001;
@@ -281,15 +285,15 @@ struct DATASMITHCONTENT_API FDatasmithTessellationOptions
 	UPROPERTY(config, EditAnywhere, BlueprintReadWrite, AdvancedDisplay, Category = "Geometry & Tessellation Options", meta = (Units = cm, ToolTip = "Tolerance used to determine if two surfaces should be stitched.", editCondition = "StitchingTechnique!=EDatasmithCADStitchingTechnique::StitchingNone"))
 	double StitchingTolerance = 0.001;
 
-	bool bUseCADKernel = false;
-
 public:
 	bool operator == (const FDatasmithTessellationOptions& Other) const
 	{
 		return FMath::IsNearlyEqual(ChordTolerance, Other.ChordTolerance)
 			&& FMath::IsNearlyEqual(MaxEdgeLength, Other.MaxEdgeLength)
 			&& FMath::IsNearlyEqual(NormalTolerance, Other.NormalTolerance)
-			&& StitchingTechnique == Other.StitchingTechnique;
+			&& StitchingTechnique == Other.StitchingTechnique
+			&& FMath::IsNearlyEqual(GeometricTolerance, Other.GeometricTolerance)
+			&& FMath::IsNearlyEqual(StitchingTolerance, Other.StitchingTolerance);
 	}
 
 	uint32 GetHash() const
@@ -301,6 +305,9 @@ public:
 		}
 		return Hash;
 	}
+	/** Helper functions to get geometrical values in the right unit, cm (native) or mm */
+	double GetGeometricTolerance( bool bInMillimeter = false ) const { return bInMillimeter ? GeometricTolerance * 10. : GeometricTolerance; }
+	double GetStitchingTolerance( bool bInMillimeter = false ) const { return bInMillimeter ? StitchingTolerance * 10. : StitchingTolerance; }
 };
 
 USTRUCT(BlueprintType)
@@ -315,7 +322,7 @@ struct DATASMITHCONTENT_API FDatasmithRetessellationOptions : public FDatasmithT
 	}
 
 	UPROPERTY(config, EditAnywhere, BlueprintReadWrite, AdvancedDisplay, Category = "Retessellation Options", meta = (ToolTip = "Regenerate deleted surfaces during retesselate or ignore them"))
-	EDatasmithCADRetessellationRule RetessellationRule;
+	EDatasmithCADRetessellationRule RetessellationRule = EDatasmithCADRetessellationRule::All;
 
 public:
 	void operator = (const FDatasmithTessellationOptions& Other)
@@ -324,6 +331,8 @@ public:
 		MaxEdgeLength = Other.MaxEdgeLength;
 		NormalTolerance = Other.NormalTolerance;
 		StitchingTechnique = Other.StitchingTechnique;
+		GeometricTolerance = Other.GetGeometricTolerance();
+		StitchingTolerance = Other.GetStitchingTolerance();
 	}
 };
 
