Move forward declaration to avoid compile error in latest msvc toolchain
#rnx

diff --git a/Engine/Plugins/Enterprise/DatasmithContent/Source/DatasmithContent/Public/DatasmithAssetUserData.h b/Engine/Plugins/Enterprise/DatasmithContent/Source/DatasmithContent/Public/DatasmithAssetUserData.h
--- a/Engine/Plugins/Enterprise/DatasmithContent/Source/DatasmithContent/Public/DatasmithAssetUserData.h
+++ b/Engine/Plugins/Enterprise/DatasmithContent/Source/DatasmithContent/Public/DatasmithAssetUserData.h
@@ -7,6 +7,8 @@
 
 #include "DatasmithAssetUserData.generated.h"
 
+class UDatasmithObjectTemplate;
+
 /** Asset user data that can be used with Datasmith on Actors and other objects  */
 UCLASS(BlueprintType, meta = (ScriptName = "DatasmithUserData", DisplayName = "Datasmith User Data"))
 class DATASMITHCONTENT_API UDatasmithAssetUserData : public UAssetUserData
@@ -23,7 +25,7 @@ public:
 
 #if WITH_EDITORONLY_DATA
 	UPROPERTY()
-	TMap< TSubclassOf< class UDatasmithObjectTemplate >, TObjectPtr<UDatasmithObjectTemplate> > ObjectTemplates;
+	TMap< TSubclassOf< UDatasmithObjectTemplate >, TObjectPtr<UDatasmithObjectTemplate> > ObjectTemplates;
 
 	virtual bool IsPostLoadThreadSafe() const override;
 	virtual void PostLoad() override;
