Deprecated UDatasmithStaticMeshCADImportData class as it is not used anymore and introduces a security vulnerability
- hackerone.com/reports #3589126

#rnx
#Synced-CL 52060097

diff --git a/Engine/Plugins/Enterprise/DatasmithCADImporter/Source/ParametricSurfaceExtension/Private/DataprepTessellationOperation.h b/Engine/Plugins/Enterprise/DatasmithCADImporter/Source/ParametricSurfaceExtension/Private/DataprepTessellationOperation.h
--- a/Engine/Plugins/Enterprise/DatasmithCADImporter/Source/ParametricSurfaceExtension/Private/DataprepTessellationOperation.h
+++ b/Engine/Plugins/Enterprise/DatasmithCADImporter/Source/ParametricSurfaceExtension/Private/DataprepTessellationOperation.h
@@ -16,7 +16,6 @@ namespace CADLibrary
 }
 struct FMeshDescription;
 class UDatasmithParametricSurfaceData;
-class UDatasmithStaticMeshCADImportData;
 
 UCLASS(Experimental, Category = MeshOperation, Meta = (DisplayName="Datasmith Tessellation", ToolTip = "For each static mesh to process, retessellate the mesh if the object contains the required data") )
 class UDataprepTessellationOperation : public UDataprepOperation
diff --git a/Engine/Plugins/Enterprise/DatasmithContent/Source/DatasmithContent/Private/DatasmithAssetImportData.cpp b/Engine/Plugins/Enterprise/DatasmithContent/Source/DatasmithContent/Private/DatasmithAssetImportData.cpp
--- a/Engine/Plugins/Enterprise/DatasmithContent/Source/DatasmithContent/Private/DatasmithAssetImportData.cpp
+++ b/Engine/Plugins/Enterprise/DatasmithContent/Source/DatasmithContent/Private/DatasmithAssetImportData.cpp
@@ -85,6 +85,7 @@ UDatasmithStaticMeshImportData* UDatasmithStaticMeshImportData::GetImportDataFor
 	return ImportData;
 }
 
+PRAGMA_DISABLE_DEPRECATION_WARNINGS
 UDatasmithStaticMeshCADImportData * UDatasmithStaticMeshCADImportData::GetCADImportDataForStaticMesh(UStaticMesh * StaticMesh, TOptional< DefaultOptionsTuple > DefaultImportOptions)
 {
 	check( StaticMesh );
@@ -151,6 +152,7 @@ const FString& UDatasmithStaticMeshCADImportData::GetResourcePath()
 
 void UDatasmithStaticMeshCADImportData::Serialize(FArchive& Ar)
 {
+	ensureMsgf(Ar.IsLoading(), TEXT("Do not use the UDatasmithStaticMeshCADImportData to store CAD data as the class is now deprecated."));
 	Super::Serialize(Ar);
 
 	Ar.UsingCustomVersion(FEnterpriseObjectVersion::GUID);
@@ -170,57 +172,11 @@ void UDatasmithStaticMeshCADImportData::Serialize(FArchive& Ar)
 		return;
 	}
 
-	FString ResourceDirectory;
 	TArray<uint8> ByteArray;
-	int32 AuxiliaryCount = AuxiliaryFilenames.Num();
-
-	if (Ar.IsSaving())
-	{
-		ResourceDirectory = FPaths::GetPath(ResourcePath);
-		FFileHelper::LoadFileToArray(ByteArray, *ResourcePath);
-	}
-	else if (Ar.IsLoading())
-	{
-		FString BaseName = FPaths::GetBaseFilename(SourceData.SourceFiles[0].RelativeFilename);
-		ResourceDirectory = FPaths::Combine(IDatasmithContentModule::Get().GetTempDir(), TEXT("ReimportCache"), *BaseName);
-		if (!FPaths::DirectoryExists(ResourceDirectory))
-		{
-			IFileManager::Get().MakeDirectory(*ResourceDirectory);
-		}
-	}
-
 	Ar << ByteArray;
-
-	if (Ar.IsSaving())
-	{
-		// Save each auxiliary file content
-		for (FString& Filename : AuxiliaryFilenames)
-		{
-			FString FilePath = FPaths::Combine(ResourceDirectory, Filename);
-
-			ByteArray.Empty();
-			FFileHelper::LoadFileToArray(ByteArray, *FilePath);
-
-			Ar << ByteArray;
-		}
-	}
-	else if (Ar.IsLoading())
-	{
-		ResourcePath = FPaths::ConvertRelativePathToFull(FPaths::Combine(ResourceDirectory, ResourceFilename));
-		FFileHelper::SaveArrayToFile(ByteArray, *ResourcePath);
-
-		// Save each auxiliary file content
-		for (FString& Filename : AuxiliaryFilenames)
-		{
-			ByteArray.Empty();
-			Ar << ByteArray;
-
-			FString FilePath = FPaths::Combine(ResourceDirectory, Filename);
-			FFileHelper::SaveArrayToFile(ByteArray, *FilePath);
-		}
-	}
 #endif // WITH_EDITORONLY_DATA
 }
+PRAGMA_ENABLE_DEPRECATION_WARNINGS
 
 // Helper function required since UObject doesn't have an interface to access AssetImportData
 UAssetImportData* Datasmith::GetAssetImportData(UObject* Asset)
diff --git a/Engine/Plugins/Enterprise/DatasmithContent/Source/DatasmithContent/Private/DatasmithContentModule.cpp b/Engine/Plugins/Enterprise/DatasmithContent/Source/DatasmithContent/Private/DatasmithContentModule.cpp
--- a/Engine/Plugins/Enterprise/DatasmithContent/Source/DatasmithContent/Private/DatasmithContentModule.cpp
+++ b/Engine/Plugins/Enterprise/DatasmithContent/Source/DatasmithContent/Private/DatasmithContentModule.cpp
@@ -20,19 +20,27 @@ class FDatasmithContentModule : public IDatasmithContentModule
 public:
 	virtual void StartupModule() override
 	{
-		// Create temporary directory which will be used by UDatasmithStaticMeshCADImportData to store transient data
-		TempDir = FPaths::Combine(FPaths::ProjectIntermediateDir(), TEXT("DatasmithContentTemp"));
-		IFileManager::Get().MakeDirectory(*TempDir);
 	}
 
 	virtual void ShutdownModule() override
 	{
-		// Clean up all transient files created during the process
-		IFileManager::Get().DeleteDirectory(*TempDir, false, true);
+		if (!TempDir.IsEmpty() && FPaths::DirectoryExists(TempDir))
+		{
+			// Clean up all transient files created during the process
+			IFileManager::Get().DeleteDirectory(*TempDir, false, true);
+		}
 	}
 
 	const FString& GetTempDir() const override
 	{
+		// For 5.8, in case users are using this method, create the temporary folder only on demand.
+		// Hence the const_cast.
+		if (TempDir.IsEmpty())
+		{
+			// Create temporary directory which will be used by UDatasmithStaticMeshCADImportData to store transient data
+			const_cast<FDatasmithContentModule*>(this)->TempDir = FPaths::Combine(FPaths::ProjectIntermediateDir(), TEXT("DatasmithContentTemp"));
+			IFileManager::Get().MakeDirectory(*TempDir);
+		}
 		return TempDir;
 	}
 
diff --git a/Engine/Plugins/Enterprise/DatasmithContent/Source/DatasmithContent/Public/DatasmithAssetImportData.h b/Engine/Plugins/Enterprise/DatasmithContent/Source/DatasmithContent/Public/DatasmithAssetImportData.h
--- a/Engine/Plugins/Enterprise/DatasmithContent/Source/DatasmithContent/Public/DatasmithAssetImportData.h
+++ b/Engine/Plugins/Enterprise/DatasmithContent/Source/DatasmithContent/Public/DatasmithAssetImportData.h
@@ -86,7 +86,8 @@ public:
 };
 
 UCLASS(MinimalAPI, HideCategories = (InternalProperty))
-class UDatasmithStaticMeshCADImportData : public UDatasmithStaticMeshImportData
+class UE_DEPRECATED(5.8, "The class UDatasmithStaticMeshCADImportData has been deprecated and will be removed next release. Please use UDatasmithParametricSurfaceData and child classes instead.")
+UDatasmithStaticMeshCADImportData : public UDatasmithStaticMeshImportData
 {
 	GENERATED_BODY()
 
diff --git a/Engine/Plugins/Enterprise/DatasmithContent/Source/DatasmithContent/Public/DatasmithContentModule.h b/Engine/Plugins/Enterprise/DatasmithContent/Source/DatasmithContent/Public/DatasmithContentModule.h
--- a/Engine/Plugins/Enterprise/DatasmithContent/Source/DatasmithContent/Public/DatasmithContentModule.h
+++ b/Engine/Plugins/Enterprise/DatasmithContent/Source/DatasmithContent/Public/DatasmithContentModule.h
@@ -41,6 +41,7 @@ public:
 	/**
 	 * @returns temporary directory created for this process
 	 */
+	UE_DEPRECATED(5.8, "IDatasmithContentModule::GetTempDir will be removed in next release.")
 	virtual const FString& GetTempDir() const = 0;
 };
 
