Fixed cooking failure duue to ensure

#rnx
#okforgithhub public
#Synced-CL 52060097

diff --git a/Engine/Plugins/Enterprise/DatasmithContent/Source/DatasmithContent/Private/DatasmithAssetImportData.cpp b/Engine/Plugins/Enterprise/DatasmithContent/Source/DatasmithContent/Private/DatasmithAssetImportData.cpp
--- a/Engine/Plugins/Enterprise/DatasmithContent/Source/DatasmithContent/Private/DatasmithAssetImportData.cpp
+++ b/Engine/Plugins/Enterprise/DatasmithContent/Source/DatasmithContent/Private/DatasmithAssetImportData.cpp
@@ -154,7 +154,7 @@ const FString& UDatasmithStaticMeshCADImportData::GetResourcePath()
 
 void UDatasmithStaticMeshCADImportData::Serialize(FArchive& Ar)
 {
-	ensureMsgf(Ar.IsLoading(), TEXT("Do not use the UDatasmithStaticMeshCADImportData to store CAD data as the class is now deprecated."));
+	UE_LOG(LogDatasmithContent, Log, TEXT("Do not use the UDatasmithStaticMeshCADImportData to store CAD data as the class is now deprecated."));
 	Super::Serialize(Ar);
 
 	Ar.UsingCustomVersion(FEnterpriseObjectVersion::GUID);
